Create a new module 'VirtualizationEditor' and move the existing VA editor code there from the 'DerivedDataEditor'

#rnx

- Up until now we have been piggybacking off the DDC editor module as the VA statistics panel was originally based on the DDC statistic panel but now we are starting to do more VA editor/UX work so we should probably have our own module.
- This does mean we need to remove the the VA statistics panel from the DDC status bar, but it didn't really belong there and the panel can still be opened via the menu 'Tool->Audit->Virtual Assets'

diff --git a/Engine/Source/Editor/DerivedDataEditor/DerivedDataEditor.Build.cs b/Engine/Source/Editor/DerivedDataEditor/DerivedDataEditor.Build.cs
--- a/Engine/Source/Editor/DerivedDataEditor/DerivedDataEditor.Build.cs
+++ b/Engine/Source/Editor/DerivedDataEditor/DerivedDataEditor.Build.cs
@@ -8,7 +8,6 @@ public class DerivedDataEditor : ModuleRules
 		 : base(Target)
 	{
 		PublicIncludePaths.Add(ModuleDirectory + "/Public");
-		PrivateIncludePaths.Add("Developer/Virtualization/Private");
 
 		PrivateDependencyModuleNames.AddRange(
 			new string[] {
@@ -27,8 +26,7 @@ public class DerivedDataEditor : ModuleRules
 				"WorkspaceMenuStructure",
 				"MessageLog",
 				"ToolWidgets",
-				"Zen",
-				"ContentBrowserData",
+				"Zen"
 			});
 
 		DynamicallyLoadedModuleNames.AddRange(
diff --git a/Engine/Source/Editor/DerivedDataEditor/Private/DerivedDataEditorModule.cpp b/Engine/Source/Editor/DerivedDataEditor/Private/DerivedDataEditorModule.cpp
--- a/Engine/Source/Editor/DerivedDataEditor/Private/DerivedDataEditorModule.cpp
+++ b/Engine/Source/Editor/DerivedDataEditor/Private/DerivedDataEditorModule.cpp
@@ -5,7 +5,6 @@
 #include "SDerivedDataStatusBar.h"
 #include "SDerivedDataDialogs.h"
 #include "SDerivedDataCacheSettings.h"
-#include "SVirtualAssetsStatistics.h"
 #include "ZenServerInterface.h"
 #include "DerivedDataCacheNotifications.h"
 #include "Modules/ModuleManager.h"
@@ -22,7 +21,6 @@ IMPLEMENT_MODULE(FDerivedDataEditorModule, DerivedDataEditor );
 
 static const FName DerivedDataResourceUsageTabName = FName(TEXT("DerivedDataResourceUsage"));
 static const FName DerivedDataCacheStatisticsTabName = FName(TEXT("DerivedDataCacheStatistics"));
-static const FName VirtualAssetsStatisticsTabName = FName(TEXT("VirtualAssetsStatistics"));
 
 void FDerivedDataEditorModule::StartupModule()
 {
@@ -42,21 +40,11 @@ void FDerivedDataEditorModule::StartupModule()
 		.SetGroup(WorkspaceMenu::GetMenuStructure().GetToolsCategory())
 		.SetIcon(CacheStatisticsIcon);
 
-	const FSlateIcon VirtaulAssetsStatisticsIcon(FAppStyle::GetAppStyleSetName(), "DerivedData.Cache.Statistics");
-
-	FGlobalTabmanager::Get()->RegisterNomadTabSpawner(VirtualAssetsStatisticsTabName, FOnSpawnTab::CreateRaw(this, &FDerivedDataEditorModule::CreateVirtualAssetsStatisticsTab))
-		.SetDisplayName(LOCTEXT("VirtualAssetsStatisticsTabTitle", "Virtual Assets"))
-		.SetTooltipText(LOCTEXT("VirtualAssetsStatisticsTabToolTipText", "Virtual Assets  Statistics"))
-		.SetGroup(WorkspaceMenu::GetMenuStructure().GetDeveloperToolsAuditCategory())
-		.SetIcon(CacheStatisticsIcon);
-
-
 #if WITH_RELOAD
 	// This code attempts to relaunch the tabs when you reload this module
 	if (IsReloadActive() && FSlateApplication::IsInitialized())
 	{
 		ShowCacheStatisticsTab();
-		ShowVirtualAssetsStatisticsTab();
 		ShowResourceUsageTab();
 	}
 #endif // WITH_RELOAD
@@ -83,13 +71,6 @@ void FDerivedDataEditorModule::ShutdownModule()
 		{
 			CacheStatisticsTab.Pin()->RequestCloseTab();
 		}
-
-		FGlobalTabmanager::Get()->UnregisterNomadTabSpawner(VirtualAssetsStatisticsTabName);
-
-		if (VirtualAssetsStatisticsTab.IsValid())
-		{
-			VirtualAssetsStatisticsTab.Pin()->RequestCloseTab();
-		}
 	}
 
 	FDerivedDataStatusBarMenuCommands::Unregister();
@@ -140,26 +121,6 @@ void FDerivedDataEditorModule::ShowCacheStatisticsTab()
 	FGlobalTabmanager::Get()->TryInvokeTab(FTabId(DerivedDataCacheStatisticsTabName));
 }
 
-
-TSharedPtr<SWidget> FDerivedDataEditorModule::CreateVirtualAssetsStatisticsDialog()
-{
-	return SNew(SVirtualAssetsStatisticsDialog);
-}
-
-TSharedRef<SDockTab> FDerivedDataEditorModule::CreateVirtualAssetsStatisticsTab(const FSpawnTabArgs& Args)
-{
-	return SAssignNew(ResourceUsageTab, SDockTab)
-		.TabRole(ETabRole::NomadTab)
-		[
-			CreateVirtualAssetsStatisticsDialog().ToSharedRef()
-		];
-}
-
-void FDerivedDataEditorModule::ShowVirtualAssetsStatisticsTab()
-{
-	FGlobalTabmanager::Get()->TryInvokeTab(FTabId(VirtualAssetsStatisticsTabName));
-}
-
 void FDerivedDataEditorModule::ShowSettingsDialog()
 {
 	if (SettingsWindow.IsValid())
diff --git a/Engine/Source/Editor/DerivedDataEditor/Private/SDerivedDataStatusBar.cpp b/Engine/Source/Editor/DerivedDataEditor/Private/SDerivedDataStatusBar.cpp
--- a/Engine/Source/Editor/DerivedDataEditor/Private/SDerivedDataStatusBar.cpp
+++ b/Engine/Source/Editor/DerivedDataEditor/Private/SDerivedDataStatusBar.cpp
@@ -69,7 +69,6 @@ void FDerivedDataStatusBarMenuCommands::RegisterCommands()
 	UI_COMMAND(ChangeSettings, "Change Cache Settings", "Opens a dialog to change Cache settings.", EUserInterfaceActionType::Button, FInputChord());
 	UI_COMMAND(ViewCacheStatistics, "View Cache Statistics", "Opens the Cache Statistics panel.", EUserInterfaceActionType::Button, FInputChord());
 	UI_COMMAND(ViewResourceUsage, "View Resource Usage", "Opens the Resource Usage panel.", EUserInterfaceActionType::Button, FInputChord());
-	UI_COMMAND(ViewVirtualAssetsStatistics, "View Virtual Assets Statistics", "Opens the Virtual Assets Statistics panel.", EUserInterfaceActionType::Button, FInputChord());
 
 	ActionList->MapAction(
 		ChangeSettings,
@@ -85,11 +84,6 @@ void FDerivedDataStatusBarMenuCommands::RegisterCommands()
 		ViewResourceUsage,
 		FExecuteAction::CreateStatic(&FDerivedDataStatusBarMenuCommands::ViewResourceUsage_Clicked)
 	);
-
-	ActionList->MapAction(
-		ViewVirtualAssetsStatistics,
-		FExecuteAction::CreateStatic(&FDerivedDataStatusBarMenuCommands::ViewVirtualAssetsStatistics_Clicked)
-	);
 }
 
 void FDerivedDataStatusBarMenuCommands::ChangeSettings_Clicked()
@@ -109,12 +103,6 @@ void FDerivedDataStatusBarMenuCommands::ViewResourceUsage_Clicked()
 	DerivedDataEditorModule.ShowResourceUsageTab();
 }
 
-void FDerivedDataStatusBarMenuCommands::ViewVirtualAssetsStatistics_Clicked()
-{
-	FDerivedDataEditorModule& DerivedDataEditorModule = FModuleManager::LoadModuleChecked<FDerivedDataEditorModule>("DerivedDataEditor");
-	DerivedDataEditorModule.ShowVirtualAssetsStatisticsTab();
-}
-
 TSharedRef<SWidget> SDerivedDataStatusBarWidget::CreateStatusBarMenu()
 {
 	UToolMenu* Menu = UToolMenus::Get()->RegisterMenu("StatusBar.ToolBar.DDC", NAME_None, EMultiBoxType::Menu, false);
@@ -146,13 +134,6 @@ TSharedRef<SWidget> SDerivedDataStatusBarWidget::CreateStatusBarMenu()
 			TAttribute<FText>(),
 			FSlateIcon(FAppStyle::GetAppStyleSetName(), "DerivedData.ResourceUsage")
 		);
-
-		Section.AddMenuEntry(
-			FDerivedDataStatusBarMenuCommands::Get().ViewVirtualAssetsStatistics,
-			TAttribute<FText>(),
-			TAttribute<FText>(),
-			FSlateIcon(FAppStyle::GetAppStyleSetName(), "DerivedData.Cache.Statistics")
-		);
 	}
 
 	return UToolMenus::Get()->GenerateWidget("StatusBar.ToolBar.DDC", FToolMenuContext(FDerivedDataStatusBarMenuCommands::ActionList));
diff --git a/Engine/Source/Editor/DerivedDataEditor/Private/SDerivedDataStatusBar.h b/Engine/Source/Editor/DerivedDataEditor/Private/SDerivedDataStatusBar.h
--- a/Engine/Source/Editor/DerivedDataEditor/Private/SDerivedDataStatusBar.h
+++ b/Engine/Source/Editor/DerivedDataEditor/Private/SDerivedDataStatusBar.h
@@ -25,19 +25,15 @@ public:
 
 private:
 
-
 	static void ChangeSettings_Clicked();
 	static void ViewCacheStatistics_Clicked();
 	static void ViewResourceUsage_Clicked();
-	static void ViewVirtualAssetsStatistics_Clicked();
-
 
 public:
 
 	TSharedPtr< FUICommandInfo > ChangeSettings;
 	TSharedPtr< FUICommandInfo > ViewResourceUsage;
 	TSharedPtr< FUICommandInfo > ViewCacheStatistics;
-	TSharedPtr< FUICommandInfo > ViewVirtualAssetsStatistics;
 
 	static TSharedRef<FUICommandList> ActionList;
 };
diff --git a/Engine/Source/Editor/DerivedDataEditor/Private/SVirtualAssetsStatistics.cpp b/Engine/Source/Editor/DerivedDataEditor/Private/SVirtualAssetsStatistics.cpp
deleted file mode 100644
diff --git a/Engine/Source/Editor/DerivedDataEditor/Private/SVirtualAssetsStatistics.h b/Engine/Source/Editor/DerivedDataEditor/Private/SVirtualAssetsStatistics.h
deleted file mode 100644
diff --git a/Engine/Source/Editor/DerivedDataEditor/Public/DerivedDataEditorModule.h b/Engine/Source/Editor/DerivedDataEditor/Public/DerivedDataEditorModule.h
--- a/Engine/Source/Editor/DerivedDataEditor/Public/DerivedDataEditorModule.h
+++ b/Engine/Source/Editor/DerivedDataEditor/Public/DerivedDataEditorModule.h
@@ -37,26 +37,21 @@ public:
 
 	void ShowResourceUsageTab();
 	void ShowCacheStatisticsTab();
-
-	void ShowVirtualAssetsStatisticsTab();
+	
 	void ShowSettingsDialog();
 
 private:
 
 	TSharedPtr<SWidget> CreateResourceUsageDialog();
 	TSharedPtr<SWidget> CreateCacheStatisticsDialog();
-	TSharedPtr<SWidget> CreateVirtualAssetsStatisticsDialog();
-
 
 	TSharedRef<SDockTab> CreateResourceUsageTab(const FSpawnTabArgs& Args);
 	TSharedRef<SDockTab> CreateCacheStatisticsTab(const FSpawnTabArgs& Args);
-	TSharedRef<SDockTab> CreateVirtualAssetsStatisticsTab(const FSpawnTabArgs& Args);
 
 	void OnSettingsDialogClosed(const TSharedRef<SWindow>& InWindow);
 
 	TWeakPtr<SDockTab> ResourceUsageTab;
 	TWeakPtr<SDockTab> CacheStatisticsTab;
-	TWeakPtr<SDockTab> VirtualAssetsStatisticsTab;
 
 	TSharedPtr<SWindow>	SettingsWindow;
 	TSharedPtr<SDerivedDataCacheSettingsDialog> SettingsDialog;
diff --git a/Engine/Source/Editor/LevelEditor/LevelEditor.Build.cs b/Engine/Source/Editor/LevelEditor/LevelEditor.Build.cs
--- a/Engine/Source/Editor/LevelEditor/LevelEditor.Build.cs
+++ b/Engine/Source/Editor/LevelEditor/LevelEditor.Build.cs
@@ -84,6 +84,7 @@ public class LevelEditor : ModuleRules
 				"DerivedDataEditor",
 				"EditorWidgets",
 				"ToolWidgets",
+				"VirtualizationEditor",
 			}
 		);
 
diff --git a/Engine/Source/Editor/UnrealEd/Private/EditorEngine.cpp b/Engine/Source/Editor/UnrealEd/Private/EditorEngine.cpp
--- a/Engine/Source/Editor/UnrealEd/Private/EditorEngine.cpp
+++ b/Engine/Source/Editor/UnrealEd/Private/EditorEngine.cpp
@@ -1074,6 +1074,7 @@ void UEditorEngine::Init(IEngineLoop* InEngineLoop)
 			TEXT("DerivedDataEditor"),
 			TEXT("CSVtoSVG"),
 			TEXT("GeometryFramework")
+			TEXT("VirtualizationEditor"),
 		};
 
 		FScopedSlowTask ModuleSlowTask(UE_ARRAY_COUNT(ModuleNames));
diff --git a/Engine/Source/Editor/VirtualizationEditor/Private/SVirtualAssetsStatistics.cpp b/Engine/Source/Editor/VirtualizationEditor/Private/SVirtualAssetsStatistics.cpp
new file mode 100644
--- /dev/null
+++ b/Engine/Source/Editor/VirtualizationEditor/Private/SVirtualAssetsStatistics.cpp
@@ -0,0 +1,658 @@
+// Copyright Epic Games, Inc. All Rights Reserved.
+
+#include "SVirtualAssetsStatistics.h"
+
+#include "Containers/Array.h"
+#include "Containers/UnrealString.h"
+#include "ContentBrowserDataMenuContexts.h"
+#include "ContentBrowserDataSource.h"
+#include "ContentBrowserItem.h"
+#include "ContentBrowserItemData.h"
+#include "Delegates/Delegate.h"
+#include "Fonts/SlateFontInfo.h"
+#include "FrameNumberTimeEvaluator.h"
+#include "Framework/Commands/UIAction.h"
+#include "Framework/Notifications/NotificationManager.h"
+#include "Framework/Text/TextLayout.h"
+#include "HAL/PlatformCrt.h"
+#include "HAL/PlatformTime.h"
+#include "Internationalization/Internationalization.h"
+#include "Internationalization/FastDecimalFormat.h"
+#include "Layout/BasicLayoutWidgetSlot.h"
+#include "Layout/Children.h"
+#include "Layout/Margin.h"
+#include "Logging/MessageLog.h"
+#include "Logging/TokenizedMessage.h"
+#include "Misc/Attribute.h"
+#include "Misc/Paths.h"
+#include "Misc/ScopeLock.h"
+#include "Settings/EditorExperimentalSettings.h"
+#include "SlotBase.h"
+#include "Styling/AppStyle.h"
+#include "Styling/CoreStyle.h"
+#include "Styling/SlateColor.h"
+#include "Styling/StyleColors.h"
+#include "Textures/SlateIcon.h"
+#include "ToolMenu.h"
+#include "ToolMenuDelegates.h"
+#include "ToolMenuSection.h"
+#include "ToolMenus.h"
+#include "Types/WidgetActiveTimerDelegate.h"
+#include "UObject/UObjectGlobals.h"
+#include "Virtualization/VirtualizationTypes.h"
+#include "Widgets/Layout/SGridPanel.h"
+#include "Widgets/Layout/SScrollBox.h"
+#include "Widgets/Notifications/SNotificationList.h"
+#include "Widgets/SBoxPanel.h"
+#include "Widgets/Text/STextBlock.h"
+
+class SWidget;
+
+#define LOCTEXT_NAMESPACE "VirtualAssets"
+
+namespace
+{
+
+FString SingleDecimalFormat(double Value)
+{
+	const FNumberFormattingOptions NumberFormattingOptions = FNumberFormattingOptions()
+		.SetUseGrouping(true)
+		.SetMinimumFractionalDigits(1)
+		.SetMaximumFractionalDigits(1);
+	return FastDecimalFormat::NumberToString(Value, ExpressionParser::GetLocalizedNumberFormattingRules(), NumberFormattingOptions);
+}
+
+} //namespace
+
+// This namespace contains code for adding the re-hydration options to the editor context menu.
+// For now it is easier to just put all virtualization editor code in this .cpp until we make
+// a virtualization specific editor module like 'DerivedDataEditor'
+namespace UE::Virtualization::Rehydration
+{
+
+/** 
+ * Attempt to rehydrate the provided files although this currently assumes that the user has
+ * checked the files out of source control before hand. 
+ */
+void RehydratePackages(TArray<FString> SelectedFiles)
+{
+	TArray<FText> Errors;
+	if (IVirtualizationSystem::Get().TryRehydratePackages(SelectedFiles, Errors) == ERehydrationResult::Success)
+	{
+		// TODO: At some point ::TryRehydratePackages will return more detail info about the process
+		// when it does we should make a better job at logging it.
+	
+		const FText Message = LOCTEXT("RehydrationSucccess", "Files were successfully re-hydrated");
+		
+		FNotificationInfo Info(Message);
+		Info.bFireAndForget = true;
+		Info.ExpireDuration = 2.0f;
+
+		FSlateNotificationManager::Get().AddNotification(Info);
+	}
+	else
+	{
+		const bool bForceNotification = true;
+
+		FMessageLog Log("LogVirtualization");
+
+		for (const FText& Error : Errors)
+		{
+			Log.Error(Error);
+		}
+
+		Log.Notify(LOCTEXT("RehydrationFailed", "Failed to rehydrate packages, see the message log for more info"), EMessageSeverity::Info, bForceNotification);
+	}
+}
+
+void BrowserItemsToFilePaths(const TArray<FContentBrowserItem>& Source, TArray<FString>& OutResult)
+{
+	for (const FContentBrowserItem& SelectedItem : Source)
+	{
+		const FContentBrowserItemData* SelectedItemData = SelectedItem.GetPrimaryInternalItem();
+		if (SelectedItemData == nullptr)
+		{
+			continue;
+		}
+
+		if (!SelectedItemData->IsFile())
+		{
+			continue;
+		}
+
+		UContentBrowserDataSource* DataSource = SelectedItemData->GetOwnerDataSource();
+		if (DataSource == nullptr)
+		{
+			continue;
+		}
+			
+		FString Path;
+		if (DataSource->GetItemPhysicalPath(*SelectedItemData, Path))
+		{
+			OutResult.Add(FPaths::ConvertRelativePathToFull(Path));
+		}
+	}
+}
+
+void AddContextMenuEntry()
+{
+	if (UToolMenu* Menu = UToolMenus::Get()->ExtendMenu("ContentBrowser.AssetContextMenu.AssetActionsSubMenu"))
+	{
+		Menu->AddDynamicSection("VirtualizedAssetsDynamic", FNewToolMenuDelegate::CreateLambda([](UToolMenu* Menu)
+		{
+			const UContentBrowserDataMenuContext_FileMenu* Context = Menu->FindContext<UContentBrowserDataMenuContext_FileMenu>();
+
+			if (Context == nullptr)
+			{
+				return;
+			}
+
+			if (Context->bCanBeModified == false)
+			{
+				return;
+			}
+
+			if (GetDefault<UEditorExperimentalSettings>()->bVirtualizedAssetRehydration == false)
+			{
+				return;
+			}
+
+			TArray<FString> SelectedFiles;
+			BrowserItemsToFilePaths(Context->SelectedItems, SelectedFiles);
+
+			if (SelectedFiles.IsEmpty())
+			{
+				return;
+			}
+
+			FToolMenuSection& Section = Menu->AddSection("VirtualizedAssets", LOCTEXT("VirtualizedAssetsHeading", "Virtualized Assets"));
+			{
+				Section.AddMenuEntry(
+					"RehydrateAsset",
+					LOCTEXT("RehydrateAsset", "Rehydrate Asset"),
+					LOCTEXT("RehydrateAssetTooltip", "Pulls the assets virtualized payloads and stores them in the package file once more"),
+					FSlateIcon(),
+					FUIAction(FExecuteAction::CreateStatic(&RehydratePackages, SelectedFiles))
+				);
+			}
+		}));
+	}
+}
+
+} // UE::Virtualization::Rehydration
+
+SVirtualAssetsStatisticsDialog::SVirtualAssetsStatisticsDialog()
+{
+	using namespace UE::Virtualization;
+	
+	// TODO - need a way to make this work once the system is initialized
+	// Register our VA notification delegate with the event
+	if (IVirtualizationSystem::IsInitialized())
+	{
+		IVirtualizationSystem& System = IVirtualizationSystem::Get();
+		System.GetNotificationEvent().AddRaw(this, &SVirtualAssetsStatisticsDialog::OnNotificationEvent);
+	}
+
+	UE::Virtualization::Rehydration::AddContextMenuEntry();
+}
+
+SVirtualAssetsStatisticsDialog::~SVirtualAssetsStatisticsDialog()
+{
+	using namespace UE::Virtualization;
+
+	// Unregister our VA notification delegate with the event
+	IVirtualizationSystem& System = IVirtualizationSystem::Get();
+	System.GetNotificationEvent().RemoveAll(this);
+}
+
+void SVirtualAssetsStatisticsDialog::OnNotificationEvent(UE::Virtualization::IVirtualizationSystem::ENotification Notification, const FIoHash& PayloadId)
+{
+	using namespace UE::Virtualization;
+
+	FScopeLock SocpeLock(&NotificationCS);
+	
+	switch (Notification)
+	{	
+		case IVirtualizationSystem::ENotification::PullBegunNotification:
+		{
+			IsPulling = true;
+			NumPullRequests++;
+
+			break;
+		}
+
+		case IVirtualizationSystem::ENotification::PullEndedNotification:
+		{	
+			if (IsPulling == true)
+			{
+				NumPullRequests--;
+				IsPulling = NumPullRequests!=0;
+			}
+			
+			break;
+		}
+
+		case IVirtualizationSystem::ENotification::PullFailedNotification:
+		{
+			NumPullRequestFailures++;
+			break;
+		}
+
+		default:
+		break;
+	}
+}
+
+void SVirtualAssetsStatisticsDialog::Construct(const FArguments& InArgs)
+{
+	this->ChildSlot
+	[
+		SNew(SVerticalBox)
+		+ SVerticalBox::Slot()
+		.AutoHeight()
+		.Padding(0, 20, 0, 0)
+		.Expose(GridSlot)
+		[
+			SAssignNew(ScrollBox, SScrollBox)
+			.Orientation(EOrientation::Orient_Horizontal)
+			.ScrollBarAlwaysVisible(false)
+
+			+ SScrollBox::Slot()
+			[
+				GetGridPanel()
+			]	
+		]
+	];
+
+	RegisterActiveTimer(0.25f, FWidgetActiveTimerDelegate::CreateSP(this, &SVirtualAssetsStatisticsDialog::UpdateGridPanels));
+}
+
+EActiveTimerReturnType SVirtualAssetsStatisticsDialog::UpdateGridPanels(double InCurrentTime, float InDeltaTime)
+{
+	ScrollBox->ClearChildren();
+	ScrollBox->AddSlot()
+		[
+			GetGridPanel()
+		];
+		
+	SlatePrepass(GetPrepassLayoutScaleMultiplier());
+
+	const float PullNotifactionTimeLimit=1.0f;
+
+	// Only show the pull notification if we have been pulling for more than a second..
+	if (NumPullRequests != 0)
+	{
+		PullNotificationTimer += InDeltaTime;
+	}
+	else
+	{
+		PullNotificationTimer = 0.0f;
+	}
+
+	if ( PullNotificationTimer>PullNotifactionTimeLimit && PullRequestNotificationItem.IsValid()==false )
+	{
+		// No existing notification or the existing one has finished
+		FNotificationInfo Info(LOCTEXT("PayloadSyncNotifcation", "Syncing Asset Payloads"));
+		Info.bFireAndForget = false;
+		Info.bUseLargeFont = false;
+		Info.bUseThrobber = false;
+		Info.FadeOutDuration = 0.5f;
+		Info.ExpireDuration = 0.0f;
+
+		PullRequestNotificationItem = FSlateNotificationManager::Get().AddNotification(Info);
+
+		if (PullRequestNotificationItem.IsValid())
+		{
+			PullRequestNotificationItem->SetCompletionState(SNotificationItem::CS_Pending);
+		}
+	}
+
+	if ( NumPullRequestFailures>0 && PullRequestFailedNotificationItem.IsValid()==false )
+	{
+		// No existing notification or the existing one has finished
+		FNotificationInfo Info(LOCTEXT("PayloadFailedNotifcation", "Failed to sync some Virtual Asset payloads from available backends.\nSome assets may no longer be usable.."));	
+		Info.bFireAndForget = false;
+		Info.bUseLargeFont = false;
+		Info.bUseThrobber = false;
+		Info.FadeOutDuration = 0.5f;
+		Info.ExpireDuration = 0.0f;
+		Info.Image = FAppStyle::GetBrush(TEXT("MessageLog.Warning"));
+		Info.ButtonDetails.Add(FNotificationButtonInfo(LOCTEXT("PullFailedIgnore", "Ignore"), LOCTEXT("PullFailedIgnoreToolTip", "Ignore future warnings"), FSimpleDelegate::CreateSP(this, &SVirtualAssetsStatisticsDialog::OnWarningReasonIgnore), SNotificationItem::CS_None));
+		Info.ButtonDetails.Add(FNotificationButtonInfo(LOCTEXT("PullFailedOK", "Ok"), LOCTEXT("PullFailedOkToolTip", "Notify future warnings"), FSimpleDelegate::CreateSP(this, &SVirtualAssetsStatisticsDialog::OnWarningReasonOk), SNotificationItem::CS_None));
+		Info.HyperlinkText = LOCTEXT("PullFailed_ShowLog", "Show Message Log");
+		Info.Hyperlink = FSimpleDelegate::CreateStatic([]() { FMessageLog("LogVirtualization").Open(EMessageSeverity::Warning, true); });
+
+		PullRequestFailedNotificationItem = FSlateNotificationManager::Get().AddNotification(Info);
+	}
+	
+	if ( NumPullRequests==0 && PullRequestNotificationItem.IsValid()==true )
+	{
+		PullRequestNotificationItem->SetCompletionState(SNotificationItem::CS_Success);
+		PullRequestNotificationItem->ExpireAndFadeout();
+		PullRequestNotificationItem.Reset();
+	}
+
+	return EActiveTimerReturnType::Continue;
+}
+
+void SVirtualAssetsStatisticsDialog::OnWarningReasonOk()
+{
+	if (PullRequestFailedNotificationItem.IsValid() == true)
+	{
+		PullRequestFailedNotificationItem->ExpireAndFadeout();
+		PullRequestFailedNotificationItem.Reset();
+		NumPullRequestFailures = 0;
+	}
+}
+
+void SVirtualAssetsStatisticsDialog::OnWarningReasonIgnore()
+{
+	if (PullRequestFailedNotificationItem.IsValid() == true)
+	{
+		PullRequestFailedNotificationItem->ExpireAndFadeout();
+	}
+}
+
+TSharedRef<SWidget> SVirtualAssetsStatisticsDialog::GetGridPanel()
+{
+	using namespace UE::Virtualization;
+
+	const float RowMargin = 0.0f;
+	const float TitleMargin = 10.0f;
+	const float ColumnMargin = 10.0f;
+	const FSlateColor TitleColor = FStyleColors::AccentWhite;
+	const FSlateFontInfo TitleFont = FCoreStyle::GetDefaultFontStyle("Bold", 10);
+	const double BytesToMegaBytes = 1.0 / (1024.0 * 1024.0);
+
+	TSharedRef<SGridPanel> Panel = SNew(SGridPanel);
+
+	// TODO - need a way to make this work once the system is initialized
+	if (!IVirtualizationSystem::IsInitialized())
+	{
+		Panel->AddSlot(0, 0)
+			[
+				SNew(STextBlock)
+				.Margin(FMargin(ColumnMargin, RowMargin))
+				.ColorAndOpacity(TitleColor)
+				.Font(TitleFont)
+				.Justification(ETextJustify::Center)
+				.Text(LOCTEXT("LazyInit", "System is uninitialized until first use"))
+			];
+
+		return Panel;
+	}
+
+	IVirtualizationSystem& System = IVirtualizationSystem::Get();	
+
+	if (System.IsEnabled() == false)
+	{
+		Panel->AddSlot(0,0)
+			[
+				SNew(STextBlock)
+				.Margin(FMargin(ColumnMargin, RowMargin))
+				.ColorAndOpacity(TitleColor)
+				.Font(TitleFont)
+				.Justification(ETextJustify::Center)
+				.Text(LOCTEXT("Disabled", "Virtual Assets Are Disabled For This Project"))
+			];
+	}
+	else
+	{
+		int32 Row = 0;
+
+		Panel->AddSlot(2, Row)
+			[
+				SNew(STextBlock)
+				.Margin(FMargin(ColumnMargin, RowMargin))
+				.ColorAndOpacity(TitleColor)
+				.Font(TitleFont)
+				.Justification(ETextJustify::Center)
+				.Text(LOCTEXT("Read", "Read"))
+			];
+
+		Panel->AddSlot(5, Row)
+			[
+				SNew(STextBlock)
+				.Margin(FMargin(ColumnMargin, RowMargin))
+				.ColorAndOpacity(TitleColor)
+				.Font(TitleFont)
+				.Justification(ETextJustify::Center)
+				.Text(LOCTEXT("Write", "Write"))
+			];
+
+		Panel->AddSlot(8, Row)
+			[
+				SNew(STextBlock)
+				.Margin(FMargin(ColumnMargin, RowMargin))
+				.ColorAndOpacity(TitleColor)
+				.Font(TitleFont)
+				.Justification(ETextJustify::Center)
+				.Text(LOCTEXT("Cache", "Cache"))
+			];
+
+		Row++;
+
+		Panel->AddSlot(0, Row)
+			[
+				SNew(STextBlock)
+				.Margin(FMargin(ColumnMargin, RowMargin, 0.0f, TitleMargin))
+				.ColorAndOpacity(TitleColor)
+				.Font(TitleFont)
+				.Justification(ETextJustify::Left)
+				.Text(LOCTEXT("Backend", "Backend"))
+			];
+
+		Panel->AddSlot(1, Row)
+			[
+				SNew(STextBlock)
+				.Margin(FMargin(ColumnMargin, RowMargin, 0.0f, TitleMargin))
+				.ColorAndOpacity(TitleColor)
+				.Font(TitleFont)
+				.Justification(ETextJustify::Center)
+				.Text(LOCTEXT("Count", "Count"))
+			];
+
+		Panel->AddSlot(2, Row)
+			[
+				SNew(STextBlock)
+				.Margin(FMargin(ColumnMargin, RowMargin, 0.0f, TitleMargin))
+				.ColorAndOpacity(TitleColor)
+				.Font(TitleFont)
+				.Justification(ETextJustify::Center)
+				.Text(LOCTEXT("Time", "Time (Sec)"))
+			];
+
+		Panel->AddSlot(3, Row)
+			[
+				SNew(STextBlock)
+				.Margin(FMargin(ColumnMargin, RowMargin, 0.0f, TitleMargin))
+				.ColorAndOpacity(TitleColor)
+				.Font(TitleFont)
+				.Justification(ETextJustify::Center)
+				.Text(LOCTEXT("Size", "Size (MB)"))
+			];
+
+		Panel->AddSlot(4, Row)
+			[
+				SNew(STextBlock)
+				.Margin(FMargin(ColumnMargin, RowMargin, 0.0f, TitleMargin))
+				.ColorAndOpacity(TitleColor)
+				.Font(TitleFont)
+				.Justification(ETextJustify::Center)
+				.Text(LOCTEXT("Count", "Count"))
+			];
+
+		Panel->AddSlot(5, Row)
+			[
+				SNew(STextBlock)
+				.Margin(FMargin(ColumnMargin, RowMargin, 0.0f, TitleMargin))
+				.ColorAndOpacity(TitleColor)
+				.Font(TitleFont)
+				.Justification(ETextJustify::Center)
+				.Text(LOCTEXT("Time", "Time (Sec)"))
+			];
+
+		Panel->AddSlot(6, Row)
+			[
+				SNew(STextBlock)
+				.Margin(FMargin(ColumnMargin, RowMargin, 0.0f, TitleMargin))
+				.ColorAndOpacity(TitleColor)
+				.Font(TitleFont)
+				.Justification(ETextJustify::Center)
+				.Text(LOCTEXT("Size", "Size (MB)"))
+			];
+
+		Panel->AddSlot(7, Row)
+			[
+				SNew(STextBlock)
+				.Margin(FMargin(ColumnMargin, RowMargin, 0.0f, TitleMargin))
+				.ColorAndOpacity(TitleColor)
+				.Font(TitleFont)
+				.Justification(ETextJustify::Center)
+				.Text(LOCTEXT("Count", "Count"))
+			];
+
+		Panel->AddSlot(8, Row)
+			[
+				SNew(STextBlock)
+				.Margin(FMargin(ColumnMargin, RowMargin, 0.0f, TitleMargin))
+				.ColorAndOpacity(TitleColor)
+				.Font(TitleFont)
+				.Justification(ETextJustify::Center)
+				.Text(LOCTEXT("Time", "Time (Sec)"))
+			];
+
+		Panel->AddSlot(9, Row)
+			[
+				SNew(STextBlock)
+				.Margin(FMargin(ColumnMargin, RowMargin, 0.0f, TitleMargin))
+				.ColorAndOpacity(TitleColor)
+				.Font(TitleFont)
+				.Justification(ETextJustify::Center)
+				.Text(LOCTEXT("Size", "Size (MB)"))
+			];
+
+		Row++;
+
+		FPayloadActivityInfo AccumulatedPayloadAcitvityInfo = System.GetAccumualtedPayloadActivityInfo();
+
+		FSlateColor Color = FStyleColors::Foreground;
+		FSlateFontInfo Font = FCoreStyle::GetDefaultFontStyle("Regular", 10);
+
+		auto DisplayPayloadActivityInfo = [&](const FString& DebugName, const FString& ConfigName, const FPayloadActivityInfo& PayloadActivityInfo)
+		{
+			Panel->AddSlot(0, Row)
+				[
+					SNew(STextBlock)
+					.Margin(FMargin(ColumnMargin, RowMargin))
+					.ColorAndOpacity(Color)
+					.Font(Font)
+					.Justification(ETextJustify::Left)
+					.Text(FText::FromString(DebugName))
+				];
+
+			Panel->AddSlot(1, Row)
+				[
+					SNew(STextBlock)
+					.Margin(FMargin(ColumnMargin, RowMargin))
+					.ColorAndOpacity(Color)
+					.Font(Font)
+					.Justification(ETextJustify::Center)
+					.Text_Lambda([PayloadActivityInfo] { return FText::FromString(FString::Printf(TEXT("%u"), PayloadActivityInfo.Pull.PayloadCount)); })
+				];
+
+			Panel->AddSlot(2, Row)
+				[
+					SNew(STextBlock)
+					.Margin(FMargin(ColumnMargin, RowMargin))
+					.ColorAndOpacity(Color)
+					.Font(Font)
+					.Justification(ETextJustify::Center)
+					.Text_Lambda([PayloadActivityInfo] { return FText::FromString(SingleDecimalFormat((double)PayloadActivityInfo.Pull.CyclesSpent * FPlatformTime::GetSecondsPerCycle())); })
+				];
+
+			Panel->AddSlot(3, Row)
+				[
+					SNew(STextBlock)
+					.Margin(FMargin(ColumnMargin, RowMargin))
+					.ColorAndOpacity(Color)
+					.Font(Font)
+					.Justification(ETextJustify::Center)
+					.Text_Lambda([PayloadActivityInfo, BytesToMegaBytes] { return FText::FromString(SingleDecimalFormat((double)PayloadActivityInfo.Pull.TotalBytes * BytesToMegaBytes)); })
+				];
+
+			Panel->AddSlot(4, Row)
+				[
+					SNew(STextBlock)
+					.Margin(FMargin(ColumnMargin, RowMargin))
+					.ColorAndOpacity(Color)
+					.Font(Font)
+					.Justification(ETextJustify::Center)
+					.Text_Lambda([PayloadActivityInfo] { return FText::FromString(FString::Printf(TEXT("%u"), PayloadActivityInfo.Push.PayloadCount)); })
+				];
+
+			Panel->AddSlot(5, Row)
+				[
+					SNew(STextBlock)
+					.Margin(FMargin(ColumnMargin, RowMargin))
+					.ColorAndOpacity(Color)
+					.Font(Font)
+					.Justification(ETextJustify::Center)
+					.Text_Lambda([PayloadActivityInfo] { return FText::FromString(SingleDecimalFormat((double)PayloadActivityInfo.Push.CyclesSpent * FPlatformTime::GetSecondsPerCycle())); })
+				];
+
+			Panel->AddSlot(6, Row)
+				[
+					SNew(STextBlock)
+					.Margin(FMargin(ColumnMargin, RowMargin))
+					.ColorAndOpacity(Color)
+					.Font(Font)
+					.Justification(ETextJustify::Center)
+					.Text_Lambda([PayloadActivityInfo, BytesToMegaBytes] { return FText::FromString(SingleDecimalFormat((double)PayloadActivityInfo.Push.TotalBytes * BytesToMegaBytes)); })
+				];
+
+			Panel->AddSlot(7, Row)
+				[
+					SNew(STextBlock)
+					.Margin(FMargin(ColumnMargin, RowMargin))
+					.ColorAndOpacity(Color)
+					.Font(Font)
+					.Justification(ETextJustify::Center)
+					.Text_Lambda([PayloadActivityInfo] { return FText::FromString(FString::Printf(TEXT("%u"), PayloadActivityInfo.Cache.PayloadCount)); })
+				];
+
+			Panel->AddSlot(8, Row)
+				[
+					SNew(STextBlock)
+					.Margin(FMargin(ColumnMargin, RowMargin))
+					.ColorAndOpacity(Color)
+					.Font(Font)
+					.Justification(ETextJustify::Center)
+					.Text_Lambda([PayloadActivityInfo] { return FText::FromString(SingleDecimalFormat((double)PayloadActivityInfo.Cache.CyclesSpent * FPlatformTime::GetSecondsPerCycle())); })
+				];
+
+			Panel->AddSlot(9, Row)
+				[
+					SNew(STextBlock)
+					.Margin(FMargin(ColumnMargin, RowMargin))
+					.ColorAndOpacity(Color)
+					.Font(Font)
+					.Justification(ETextJustify::Center)
+					.Text_Lambda([PayloadActivityInfo, BytesToMegaBytes] { return FText::FromString(SingleDecimalFormat((double)PayloadActivityInfo.Cache.TotalBytes * BytesToMegaBytes)); })
+				];
+
+			Row++;
+		};
+
+		System.GetPayloadActivityInfo(DisplayPayloadActivityInfo);
+
+		Color = TitleColor;
+		Font = TitleFont;
+
+		DisplayPayloadActivityInfo(FString("Total"), FString("Total"), AccumulatedPayloadAcitvityInfo);
+	}
+
+	return Panel;
+}
+
+#undef LOCTEXT_NAMESPACE
diff --git a/Engine/Source/Editor/VirtualizationEditor/Private/SVirtualAssetsStatistics.h b/Engine/Source/Editor/VirtualizationEditor/Private/SVirtualAssetsStatistics.h
new file mode 100644
--- /dev/null
+++ b/Engine/Source/Editor/VirtualizationEditor/Private/SVirtualAssetsStatistics.h
@@ -0,0 +1,55 @@
+// Copyright Epic Games, Inc. All Rights Reserved.
+
+#pragma once
+
+#include "HAL/CriticalSection.h"
+#include "HAL/Platform.h"
+#include "IO/IoHash.h"
+#include "Internationalization/Text.h"
+#include "Templates/SharedPointer.h"
+#include "Types/SlateEnums.h"
+#include "Virtualization/VirtualizationSystem.h"
+#include "Widgets/DeclarativeSyntaxSupport.h"
+#include "Widgets/SBoxPanel.h"
+#include "Widgets/SCompoundWidget.h"
+
+class SNotificationItem;
+class SWidget;
+
+class SVirtualAssetsStatisticsDialog : public SCompoundWidget
+{
+	SLATE_BEGIN_ARGS(SVirtualAssetsStatisticsDialog) {}
+	SLATE_END_ARGS()
+
+	void Construct(const FArguments& InArgs);
+
+	SVirtualAssetsStatisticsDialog();
+	virtual ~SVirtualAssetsStatisticsDialog();
+
+private:
+
+	TSharedRef<SWidget> GetGridPanel();
+
+	EActiveTimerReturnType UpdateGridPanels(double InCurrentTime, float InDeltaTime);
+
+	FText GetNotificationText() const;
+
+	void OnNotificationEvent(UE::Virtualization::IVirtualizationSystem::ENotification Notification, const FIoHash& PayloadId);
+	void OnWarningReasonOk();
+	void OnWarningReasonIgnore();
+
+	SVerticalBox::FSlot* GridSlot = nullptr;
+
+	FCriticalSection NotificationCS;
+
+	TSharedPtr<SNotificationItem> PullRequestNotificationItem;
+	TSharedPtr<SNotificationItem> PullRequestFailedNotificationItem;
+
+	TSharedPtr<class SScrollBox> ScrollBox;
+
+	bool	IsPulling = false;
+	float	PullNotificationTimer = 0.0f;
+	uint32	NumPullRequests = 0;
+	uint32	NumPullRequestFailures = 0;
+};
+
diff --git a/Engine/Source/Editor/VirtualizationEditor/Private/VirtualizationEditorModule.cpp b/Engine/Source/Editor/VirtualizationEditor/Private/VirtualizationEditorModule.cpp
new file mode 100644
--- /dev/null
+++ b/Engine/Source/Editor/VirtualizationEditor/Private/VirtualizationEditorModule.cpp
@@ -0,0 +1,69 @@
+// Copyright Epic Games, Inc. All Rights Reserved.
+
+#include "VirtualizationEditorModule.h"
+
+#include "Editor/WorkspaceMenuStructure/Public/WorkspaceMenuStructure.h"
+#include "Editor/WorkspaceMenuStructure/Public/WorkspaceMenuStructureModule.h"
+#include "Framework/Application/SlateApplication.h"
+#include "Framework/Docking/TabManager.h"
+#include "SVirtualAssetsStatistics.h"
+#include "Textures/SlateIcon.h"
+#include "Widgets/Docking/SDockTab.h"
+
+#define LOCTEXT_NAMESPACE "VirtualizationEditor"
+
+IMPLEMENT_MODULE(FVirtualizationEditorModule, VirtualizationEditor);
+
+static const FName VirtualAssetsStatisticsTabName = FName(TEXT("VirtualAssetsStatistics"));
+
+void FVirtualizationEditorModule::StartupModule()
+{
+	const FSlateIcon VirtaulAssetsStatisticsIcon(FAppStyle::GetAppStyleSetName(), "DerivedData.Cache.Statistics");
+	
+	FGlobalTabmanager::Get()->RegisterNomadTabSpawner(VirtualAssetsStatisticsTabName, FOnSpawnTab::CreateRaw(this, &FVirtualizationEditorModule::CreateVirtualAssetsStatisticsTab))
+		.SetDisplayName(LOCTEXT("VirtualAssetsStatisticsTabTitle", "Virtual Assets"))
+		.SetTooltipText(LOCTEXT("VirtualAssetsStatisticsTabToolTipText", "Virtual Assets  Statistics"))
+		.SetGroup(WorkspaceMenu::GetMenuStructure().GetDeveloperToolsAuditCategory())
+		.SetIcon(VirtaulAssetsStatisticsIcon);
+
+#if WITH_RELOAD
+	// This code attempts to relaunch the tabs when you reload this module
+	if (IsReloadActive() && FSlateApplication::IsInitialized())
+	{
+		ShowVirtualAssetsStatisticsTab();
+	}
+#endif // WITH_RELOAD
+
+}
+
+void FVirtualizationEditorModule::ShutdownModule()
+{
+	if (FSlateApplication::IsInitialized())
+	{
+		if (VirtualAssetsStatisticsTab.IsValid())
+		{
+			VirtualAssetsStatisticsTab.Pin()->RequestCloseTab();
+		}
+
+		FGlobalTabmanager::Get()->UnregisterNomadTabSpawner(VirtualAssetsStatisticsTabName);
+	}
+}
+
+void FVirtualizationEditorModule::ShowVirtualAssetsStatisticsTab()
+{
+	FGlobalTabmanager::Get()->TryInvokeTab(FTabId(VirtualAssetsStatisticsTabName));
+}
+
+TSharedPtr<SWidget> FVirtualizationEditorModule::CreateVirtualAssetsStatisticsDialog()
+{
+	return SNew(SVirtualAssetsStatisticsDialog);
+}
+
+TSharedRef<SDockTab> FVirtualizationEditorModule::CreateVirtualAssetsStatisticsTab(const FSpawnTabArgs& Args)
+{
+	return SAssignNew(VirtualAssetsStatisticsTab, SDockTab)
+		.TabRole(ETabRole::NomadTab)
+		[
+			CreateVirtualAssetsStatisticsDialog().ToSharedRef()
+		];
+}
diff --git a/Engine/Source/Editor/VirtualizationEditor/Public/VirtualizationEditorModule.h b/Engine/Source/Editor/VirtualizationEditor/Public/VirtualizationEditorModule.h
new file mode 100644
--- /dev/null
+++ b/Engine/Source/Editor/VirtualizationEditor/Public/VirtualizationEditorModule.h
@@ -0,0 +1,28 @@
+// Copyright Epic Games, Inc. All Rights Reserved.
+
+#pragma once
+
+#include "Modules/ModuleInterface.h"
+#include "Templates/SharedPointer.h"
+
+class SWidget;
+class SDockTab;
+class FSpawnTabArgs;
+
+/**
+ * The module holding all of the UI related pieces for the Virtualization module
+ */
+class FVirtualizationEditorModule : public IModuleInterface
+{
+private:
+
+	virtual void StartupModule() override;
+	virtual void ShutdownModule() override;
+
+	void ShowVirtualAssetsStatisticsTab();
+
+	TSharedPtr<SWidget> CreateVirtualAssetsStatisticsDialog();
+	TSharedRef<SDockTab> CreateVirtualAssetsStatisticsTab(const FSpawnTabArgs& Args);
+
+	TWeakPtr<SDockTab> VirtualAssetsStatisticsTab;
+};
diff --git a/Engine/Source/Editor/VirtualizationEditor/VirtualizationEditor.Build.cs b/Engine/Source/Editor/VirtualizationEditor/VirtualizationEditor.Build.cs
new file mode 100644
--- /dev/null
+++ b/Engine/Source/Editor/VirtualizationEditor/VirtualizationEditor.Build.cs
@@ -0,0 +1,24 @@
+// Copyright Epic Games, Inc. All Rights Reserved.
+
+using UnrealBuildTool;
+
+public class VirtualizationEditor : ModuleRules
+{
+	public VirtualizationEditor(ReadOnlyTargetRules Target)
+		 : base(Target)
+	{
+		PrivateDependencyModuleNames.AddRange(
+			new string[]
+			{
+				"ContentBrowserData",
+				"Core",
+				"CoreUObject",
+				"Engine",
+				"UnrealEd",
+				"ToolMenus",
+				"Slate",
+				"SlateCore",
+				"Virtualization"
+			});
+	}
+}
