CIS fix - Add missing undef of LOCTEXT_NAMESPACE

#rnx

diff --git a/Engine/Source/Editor/VirtualizationEditor/Private/VirtualizationEditorModule.cpp b/Engine/Source/Editor/VirtualizationEditor/Private/VirtualizationEditorModule.cpp
--- a/Engine/Source/Editor/VirtualizationEditor/Private/VirtualizationEditorModule.cpp
+++ b/Engine/Source/Editor/VirtualizationEditor/Private/VirtualizationEditorModule.cpp
@@ -67,3 +67,5 @@ TSharedRef<SDockTab> FVirtualizationEditorModule::CreateVirtualAssetsStatisticsT
 			CreateVirtualAssetsStatisticsDialog().ToSharedRef()
 		];
 }
+
+#undef LOCTEXT_NAMESPACE
