Move the rehydration context menu functionality from the statistics panel to its own code file

#rnx

- When the statistics panel lived in the DDC editor module it was easier to keep all virtualization functionality in one place. Now that we have our own module we can move it to its own code file for better readability.
- Added a utility function ::SetupMenuEntries called when the module is started, we should add all future menu entries here.

diff --git a/Engine/Source/Editor/VirtualizationEditor/Private/RehydrationMenuEntry.cpp b/Engine/Source/Editor/VirtualizationEditor/Private/RehydrationMenuEntry.cpp
new file mode 100644
--- /dev/null
+++ b/Engine/Source/Editor/VirtualizationEditor/Private/RehydrationMenuEntry.cpp
@@ -0,0 +1,144 @@
+// Copyright Epic Games, Inc. All Rights Reserved.
+
+#include "RehydrationMenuEntry.h"
+
+#include "ContentBrowserDataMenuContexts.h"
+#include "ContentBrowserDataSource.h"
+#include "ContentBrowserItem.h"
+#include "ContentBrowserItemData.h"
+#include "Framework/Notifications/NotificationManager.h"
+#include "Logging/MessageLog.h"
+#include "Settings/EditorExperimentalSettings.h"
+#include "Styling/AppStyle.h"
+#include "ToolMenus.h"
+#include "Virtualization/VirtualizationSystem.h"
+#include "Widgets/Notifications/SNotificationList.h"
+
+#define LOCTEXT_NAMESPACE "VirtualizationEditor"
+
+namespace UE::Virtualization
+{
+
+/**
+ * Attempt to rehydrate the provided files although this currently assumes that the user has
+ * checked the files out of source control before hand.
+ */
+void RehydratePackages(TArray<FString> SelectedFiles)
+{
+	TArray<FText> Errors;
+	if (IVirtualizationSystem::Get().TryRehydratePackages(SelectedFiles, Errors) == ERehydrationResult::Success)
+	{
+		// TODO: At some point ::TryRehydratePackages will return more detail info about the process
+		// when it does we should make a better job at logging it.
+
+		const FText Message = LOCTEXT("RehydrationSucccess", "Files were successfully re-hydrated");
+
+		FNotificationInfo Info(Message);
+		Info.bFireAndForget = true;
+		Info.ExpireDuration = 2.0f;
+
+		FSlateNotificationManager::Get().AddNotification(Info);
+	}
+	else
+	{
+		const bool bForceNotification = true;
+
+		FMessageLog Log("LogVirtualization");
+
+		for (const FText& Error : Errors)
+		{
+			Log.Error(Error);
+		}
+
+		Log.Notify(LOCTEXT("RehydrationFailed", "Failed to rehydrate packages, see the message log for more info"), EMessageSeverity::Info, bForceNotification);
+	}
+}
+
+void BrowserItemsToFilePaths(const TArray<FContentBrowserItem>& Source, TArray<FString>& OutResult)
+{
+	for (const FContentBrowserItem& SelectedItem : Source)
+	{
+		const FContentBrowserItemData* SelectedItemData = SelectedItem.GetPrimaryInternalItem();
+		if (SelectedItemData == nullptr)
+		{
+			continue;
+		}
+
+		if (!SelectedItemData->IsFile())
+		{
+			continue;
+		}
+
+		UContentBrowserDataSource* DataSource = SelectedItemData->GetOwnerDataSource();
+		if (DataSource == nullptr)
+		{
+			continue;
+		}
+
+		FString Path;
+		if (DataSource->GetItemPhysicalPath(*SelectedItemData, Path))
+		{
+			OutResult.Add(FPaths::ConvertRelativePathToFull(Path));
+		}
+	}
+}
+
+void SetupRehydrationContentMenuEntry()
+{
+	UToolMenus* ToolsMenu = UToolMenus::Get();
+	if (ToolsMenu == nullptr)
+	{
+		return;
+	}
+
+	UToolMenu* Menu = UToolMenus::Get()->ExtendMenu("ContentBrowser.AssetContextMenu.AssetActionsSubMenu");
+	if (Menu == nullptr)
+	{
+		return;
+	}
+
+	Menu->AddDynamicSection("VirtualizedAssetsDynamic", FNewToolMenuDelegate::CreateLambda([](UToolMenu* Menu)
+		{
+			const UContentBrowserDataMenuContext_FileMenu* Context = Menu->FindContext<UContentBrowserDataMenuContext_FileMenu>();
+
+			if (Context == nullptr)
+			{
+				return;
+			}
+
+			if (Context->bCanBeModified == false)
+			{
+				return;
+			}
+
+			// The user needs to enable this experimental editor setting before we display the
+			// option in the context menu. 
+			if (GetDefault<UEditorExperimentalSettings>()->bVirtualizedAssetRehydration == false)
+			{
+				return;
+			}
+
+			TArray<FString> SelectedFiles;
+			BrowserItemsToFilePaths(Context->SelectedItems, SelectedFiles);
+
+			if (SelectedFiles.IsEmpty())
+			{
+				return;
+			}
+
+			FToolMenuSection& Section = Menu->AddSection("VirtualizedAssets", LOCTEXT("VirtualizedAssetsHeading", "Virtualized Assets"));
+			{
+				Section.AddMenuEntry(
+					"RehydrateAsset",
+					LOCTEXT("RehydrateAsset", "Rehydrate Asset"),
+					LOCTEXT("RehydrateAssetTooltip", "Pulls the assets virtualized payloads and stores them in the package file once more"),
+					FSlateIcon(),
+					FUIAction(FExecuteAction::CreateStatic(&RehydratePackages, SelectedFiles))
+				);
+			}
+		}));
+}
+
+} // namespace UE::Virtualization
+
+#undef LOCTEXT_NAMESPACE
diff --git a/Engine/Source/Editor/VirtualizationEditor/Private/RehydrationMenuEntry.h b/Engine/Source/Editor/VirtualizationEditor/Private/RehydrationMenuEntry.h
new file mode 100644
--- /dev/null
+++ b/Engine/Source/Editor/VirtualizationEditor/Private/RehydrationMenuEntry.h
@@ -0,0 +1,11 @@
+// Copyright Epic Games, Inc. All Rights Reserved.
+
+#pragma once
+
+namespace UE::Virtualization
+{
+
+/** Adds an option to rehydrate packages under the asset context menu */
+void SetupRehydrationContentMenuEntry();
+
+} // namespace UE::Virtualization
diff --git a/Engine/Source/Editor/VirtualizationEditor/Private/SVirtualAssetsStatistics.cpp b/Engine/Source/Editor/VirtualizationEditor/Private/SVirtualAssetsStatistics.cpp
--- a/Engine/Source/Editor/VirtualizationEditor/Private/SVirtualAssetsStatistics.cpp
+++ b/Engine/Source/Editor/VirtualizationEditor/Private/SVirtualAssetsStatistics.cpp
@@ -4,10 +4,6 @@
 
 #include "Containers/Array.h"
 #include "Containers/UnrealString.h"
-#include "ContentBrowserDataMenuContexts.h"
-#include "ContentBrowserDataSource.h"
-#include "ContentBrowserItem.h"
-#include "ContentBrowserItemData.h"
 #include "Delegates/Delegate.h"
 #include "Fonts/SlateFontInfo.h"
 #include "FrameNumberTimeEvaluator.h"
@@ -26,9 +22,7 @@
 #include "Misc/Attribute.h"
 #include "Misc/Paths.h"
 #include "Misc/ScopeLock.h"
-#include "Settings/EditorExperimentalSettings.h"
 #include "SlotBase.h"
-#include "Styling/AppStyle.h"
 #include "Styling/CoreStyle.h"
 #include "Styling/SlateColor.h"
 #include "Styling/StyleColors.h"
@@ -36,7 +30,6 @@
 #include "ToolMenu.h"
 #include "ToolMenuDelegates.h"
 #include "ToolMenuSection.h"
-#include "ToolMenus.h"
 #include "Types/WidgetActiveTimerDelegate.h"
 #include "UObject/UObjectGlobals.h"
 #include "Virtualization/VirtualizationTypes.h"
@@ -48,7 +41,7 @@
 
 class SWidget;
 
-#define LOCTEXT_NAMESPACE "VirtualAssets"
+#define LOCTEXT_NAMESPACE "VirtualizationEditor"
 
 namespace
 {
@@ -64,123 +57,6 @@ FString SingleDecimalFormat(double Value)
 
 } //namespace
 
-// This namespace contains code for adding the re-hydration options to the editor context menu.
-// For now it is easier to just put all virtualization editor code in this .cpp until we make
-// a virtualization specific editor module like 'DerivedDataEditor'
-namespace UE::Virtualization::Rehydration
-{
-
-/** 
- * Attempt to rehydrate the provided files although this currently assumes that the user has
- * checked the files out of source control before hand. 
- */
-void RehydratePackages(TArray<FString> SelectedFiles)
-{
-	TArray<FText> Errors;
-	if (IVirtualizationSystem::Get().TryRehydratePackages(SelectedFiles, Errors) == ERehydrationResult::Success)
-	{
-		// TODO: At some point ::TryRehydratePackages will return more detail info about the process
-		// when it does we should make a better job at logging it.
-	
-		const FText Message = LOCTEXT("RehydrationSucccess", "Files were successfully re-hydrated");
-		
-		FNotificationInfo Info(Message);
-		Info.bFireAndForget = true;
-		Info.ExpireDuration = 2.0f;
-
-		FSlateNotificationManager::Get().AddNotification(Info);
-	}
-	else
-	{
-		const bool bForceNotification = true;
-
-		FMessageLog Log("LogVirtualization");
-
-		for (const FText& Error : Errors)
-		{
-			Log.Error(Error);
-		}
-
-		Log.Notify(LOCTEXT("RehydrationFailed", "Failed to rehydrate packages, see the message log for more info"), EMessageSeverity::Info, bForceNotification);
-	}
-}
-
-void BrowserItemsToFilePaths(const TArray<FContentBrowserItem>& Source, TArray<FString>& OutResult)
-{
-	for (const FContentBrowserItem& SelectedItem : Source)
-	{
-		const FContentBrowserItemData* SelectedItemData = SelectedItem.GetPrimaryInternalItem();
-		if (SelectedItemData == nullptr)
-		{
-			continue;
-		}
-
-		if (!SelectedItemData->IsFile())
-		{
-			continue;
-		}
-
-		UContentBrowserDataSource* DataSource = SelectedItemData->GetOwnerDataSource();
-		if (DataSource == nullptr)
-		{
-			continue;
-		}
-			
-		FString Path;
-		if (DataSource->GetItemPhysicalPath(*SelectedItemData, Path))
-		{
-			OutResult.Add(FPaths::ConvertRelativePathToFull(Path));
-		}
-	}
-}
-
-void AddContextMenuEntry()
-{
-	if (UToolMenu* Menu = UToolMenus::Get()->ExtendMenu("ContentBrowser.AssetContextMenu.AssetActionsSubMenu"))
-	{
-		Menu->AddDynamicSection("VirtualizedAssetsDynamic", FNewToolMenuDelegate::CreateLambda([](UToolMenu* Menu)
-		{
-			const UContentBrowserDataMenuContext_FileMenu* Context = Menu->FindContext<UContentBrowserDataMenuContext_FileMenu>();
-
-			if (Context == nullptr)
-			{
-				return;
-			}
-
-			if (Context->bCanBeModified == false)
-			{
-				return;
-			}
-
-			if (GetDefault<UEditorExperimentalSettings>()->bVirtualizedAssetRehydration == false)
-			{
-				return;
-			}
-
-			TArray<FString> SelectedFiles;
-			BrowserItemsToFilePaths(Context->SelectedItems, SelectedFiles);
-
-			if (SelectedFiles.IsEmpty())
-			{
-				return;
-			}
-
-			FToolMenuSection& Section = Menu->AddSection("VirtualizedAssets", LOCTEXT("VirtualizedAssetsHeading", "Virtualized Assets"));
-			{
-				Section.AddMenuEntry(
-					"RehydrateAsset",
-					LOCTEXT("RehydrateAsset", "Rehydrate Asset"),
-					LOCTEXT("RehydrateAssetTooltip", "Pulls the assets virtualized payloads and stores them in the package file once more"),
-					FSlateIcon(),
-					FUIAction(FExecuteAction::CreateStatic(&RehydratePackages, SelectedFiles))
-				);
-			}
-		}));
-	}
-}
-
-} // UE::Virtualization::Rehydration
-
 SVirtualAssetsStatisticsDialog::SVirtualAssetsStatisticsDialog()
 {
 	using namespace UE::Virtualization;
@@ -192,8 +68,6 @@ SVirtualAssetsStatisticsDialog::SVirtualAssetsStatisticsDialog()
 		IVirtualizationSystem& System = IVirtualizationSystem::Get();
 		System.GetNotificationEvent().AddRaw(this, &SVirtualAssetsStatisticsDialog::OnNotificationEvent);
 	}
-
-	UE::Virtualization::Rehydration::AddContextMenuEntry();
 }
 
 SVirtualAssetsStatisticsDialog::~SVirtualAssetsStatisticsDialog()
diff --git a/Engine/Source/Editor/VirtualizationEditor/Private/VirtualizationEditorModule.cpp b/Engine/Source/Editor/VirtualizationEditor/Private/VirtualizationEditorModule.cpp
--- a/Engine/Source/Editor/VirtualizationEditor/Private/VirtualizationEditorModule.cpp
+++ b/Engine/Source/Editor/VirtualizationEditor/Private/VirtualizationEditorModule.cpp
@@ -6,6 +6,8 @@
 #include "Editor/WorkspaceMenuStructure/Public/WorkspaceMenuStructureModule.h"
 #include "Framework/Application/SlateApplication.h"
 #include "Framework/Docking/TabManager.h"
+#include "Modules/ModuleManager.h"
+#include "RehydrationMenuEntry.h"
 #include "SVirtualAssetsStatistics.h"
 #include "Textures/SlateIcon.h"
 #include "Widgets/Docking/SDockTab.h"
@@ -16,6 +18,17 @@ IMPLEMENT_MODULE(FVirtualizationEditorModule, VirtualizationEditor);
 
 static const FName VirtualAssetsStatisticsTabName = FName(TEXT("VirtualAssetsStatistics"));
 
+namespace 
+{
+
+/** Utility function for adding menu entries, called when the module is started */
+void SetupMenuEntries()
+{
+	UE::Virtualization::SetupRehydrationContentMenuEntry();
+}
+
+} //namespace
+
 void FVirtualizationEditorModule::StartupModule()
 {
 	const FSlateIcon VirtaulAssetsStatisticsIcon(FAppStyle::GetAppStyleSetName(), "DerivedData.Cache.Statistics");
@@ -34,6 +47,7 @@ void FVirtualizationEditorModule::StartupModule()
 	}
 #endif // WITH_RELOAD
 
+	SetupMenuEntries();
 }
 
 void FVirtualizationEditorModule::ShutdownModule()
