Add a very basic way to precache all virtualized payloads in a project and store them in the cache storage backends (usually the local DDC in our internal set up)
#rnx
The commandlet is run via the cmdline '-run=VirtualizationEditor.PrecachePayloads' (the VirtualizationEditor part will load the module before executing the PrecachePayloads commandlet)
At the moment this is very slow to run as the VA API only allows one payload to be pulled in a single call and although we can make requests on many threads, there is a lot of overhead when requesting data from some backends and a batch pull API would be substantially quicker. This will be looked at in a future work item.
There are currently 4 different ways we can iterate over the payloads and precache them, since the commandlet is still too slow to use in general production, the different ways of iterating are provided so we can experiment later.
There is no way to customize what is precached, we just iterate over all the packaged in a project. We might want to improve this in the future but it would perhaps make more sense to move the package iteration options from the resave commandlet and make it something general purpose that different commandlets can also use. This would help build up a consistent way of choosing which packages are affected by a commandlet.

diff --git a/Engine/Source/Editor/VirtualizationEditor/Private/PrecachePayloads.cpp b/Engine/Source/Editor/VirtualizationEditor/Private/PrecachePayloads.cpp
new file mode 100644
--- /dev/null
+++ b/Engine/Source/Editor/VirtualizationEditor/Private/PrecachePayloads.cpp
@@ -0,0 +1,284 @@
+// Copyright Epic Games, Inc. All Rights Reserved.
+
+#include "PrecachePayloads.h"
+
+#include "PackageHelperFunctions.h"
+#include "Virtualization/VirtualizationSystem.h"
+#include "UObject/PackageTrailer.h"
+#include "Async/ParallelFor.h"
+
+// 0 - No threading
+// 1 - ParallelFor
+// 2 - Split the work into equal sized batches across all cores
+// 3 - Use a limited number of worker tasks
+
+#define UE_ITERATION_METHOD 3
+
+namespace 
+{
+
+class FWorkQueue
+{
+public:
+	using FJob = TArrayView<const FIoHash>;
+
+	FWorkQueue(TArray<FIoHash>&& InWork, int32 JobSize)
+		: Work(InWork)
+	{
+		const int32 NumJobs = FMath::DivideAndRoundUp<int32>(Work.Num(), JobSize);
+		Jobs.Reserve(NumJobs);
+
+		for (int32 JobIndex = 0; JobIndex < NumJobs; ++JobIndex)
+		{
+			const int32 JobStart = JobIndex * JobSize;
+			const int32 JobEnd = FMath::Min((JobIndex + 1) * JobSize, Work.Num());
+
+			FJob Job = MakeArrayView(&Work[JobStart], JobEnd - JobStart);
+			Jobs.Add(Job);
+		}
+	}
+
+	FJob GetJob()
+	{
+		TRACE_CPUPROFILER_EVENT_SCOPE(FWorkQueue::GetJob);
+
+		FScopeLock _(&JobsCS);
+
+		if (!Jobs.IsEmpty())
+		{
+			return Jobs.Pop(false);
+		}
+		else
+		{
+			return FJob();
+		}	
+	}
+
+
+	TArray<FIoHash> Work;
+	TArray<FJob> Jobs;
+
+	FCriticalSection JobsCS;
+};
+
+/** 
+ * Utility to wait for the given tasks to complete and return early if the tasks are not completed
+ * within the given timespan.
+ * This is used isntead of UE::Tasks::Wait, because that call can execute the remaining tasks if not
+ * already being processed, which in most scenarios makes sense, but since we are running in a commandlet
+ * and we know nothing else is running, we don't need the game thread processing work as well. 
+ */
+bool WaitOnTasks(const TArray<UE::Tasks::FTask>& Tasks, FTimespan InTimeout)
+{
+	UE::FTimeout Timeout{ InTimeout };
+
+	FSharedEventRef CompletionEvent;
+
+	UE::Tasks::FTask WaitingTask = UE::Tasks::Launch(UE_SOURCE_LOCATION,
+		[CompletionEvent]
+		{
+			CompletionEvent->Trigger();
+		},
+		Tasks, LowLevelTasks::ETaskPriority::Default, UE::Tasks::EExtendedTaskPriority::Inline);
+
+	return CompletionEvent->Wait((uint32)FMath::Clamp<int64>(Timeout.GetRemainingTime().GetTicks() / ETimespan::TicksPerMillisecond, 0, MAX_uint32));
+}
+
+/** Parse all of the active mount points and find all .uasset/.umaps */
+TArray<FString> FindPackages()
+{
+	TRACE_CPUPROFILER_EVENT_SCOPE(FindPackages);
+
+	TArray<FString> PackageNames;
+
+	// TODO: In theory we could check the VA settings and skip engine content if they are 
+	// excluded from virtualization. Not convinced it saves enough time to be worth the 
+	// effort.
+	EPackageNormalizationFlags PackageFilter = NORMALIZE_DefaultFlags;
+
+	const FString AssetSearch = TEXT("*") + FPackageName::GetAssetPackageExtension();
+	const FString MapSearch = TEXT("*") + FPackageName::GetMapPackageExtension();
+	
+	TArray<FString> Unused;
+	bool bAnyFound = NormalizePackageNames(Unused, PackageNames, AssetSearch, PackageFilter);
+	bAnyFound |= NormalizePackageNames(Unused, PackageNames, MapSearch, PackageFilter);
+
+	return PackageNames;
+}
+
+/** Returns a combined list of unique virtualized payload ids from the given list of packages */
+TArray<FIoHash> FindVirtualizedPayloads(const TArray<FString>& PackageNames)
+{
+	TRACE_CPUPROFILER_EVENT_SCOPE(FindVirtualizedPayloads);
+
+	TSet<FIoHash> PayloadsToPrecache;
+
+	for (const FString& PackageName : PackageNames)
+	{
+		UE::FPackageTrailer Trailer;
+		if (UE::FPackageTrailer::TryLoadFromFile(PackageName, Trailer))
+		{
+			TArray<FIoHash> VirtualizedPayloads = Trailer.GetPayloads(UE::EPayloadStorageType::Virtualized);
+			PayloadsToPrecache.Append(VirtualizedPayloads);
+		}
+	}
+
+	return PayloadsToPrecache.Array();
+}
+
+}
+
+UPrecachePayloadsCommandlet::UPrecachePayloadsCommandlet(const FObjectInitializer& ObjectInitializer)
+	: Super(ObjectInitializer)
+{
+}
+
+int32 UPrecachePayloadsCommandlet::Main(const FString& Params)
+{
+	TRACE_CPUPROFILER_EVENT_SCOPE(UPrecachePayloadsCommandlet);
+
+	TArray<FString> PackageNames = FindPackages();
+
+	UE_LOG(LogVirtualization, Display, TEXT("Found %d packages"), PackageNames.Num());
+
+	TArray<FIoHash> PayloadIds = FindVirtualizedPayloads(PackageNames);
+
+	if (PayloadIds.IsEmpty())
+	{
+		UE_LOG(LogVirtualization, Display, TEXT("No virtualized payloads found"));
+		return 0;
+	}
+		
+	UE_LOG(LogVirtualization, Display, TEXT("Found %d virtualized payloads to precache"), PayloadIds.Num());
+	UE_LOG(LogVirtualization, Display, TEXT("Precaching payloads..."));
+
+	UE::Virtualization::IVirtualizationSystem& System = UE::Virtualization::IVirtualizationSystem::Get();
+
+#if UE_ITERATION_METHOD == 0
+	{
+		for (const FIoHash& PayloadId : PayloadIds)
+		{
+			FCompressedBuffer Payload = System.PullData(PayloadId);
+			if (Payload.IsNull())
+			{
+				UE_LOG(LogVirtualization, Error, TEXT("Failed to precache payload '%s'"), *LexToString(PayloadId));
+			}
+		}
+	}
+#elif UE_ITERATION_METHOD == 1
+	{
+		ParallelFor(PayloadIds.Num(), [&PayloadIds, &System](int32 Index)
+			{
+				const FIoHash& PayloadId = PayloadIds[Index];
+
+				FCompressedBuffer Payload = System.PullData(PayloadId);
+				if (Payload.IsNull())
+				{
+					UE_LOG(LogVirtualization, Error, TEXT("Failed to precache payload '%s'"), *LexToString(PayloadId));
+				}
+			});
+	}
+#elif UE_ITERATION_METHOD == 2
+	{
+		int32 NumBatches = FPlatformMisc::NumberOfCores();
+
+		TArray<UE::Tasks::FTask> Tasks;
+		Tasks.Reserve(NumBatches);
+		
+		std::atomic<int32> NumCompletedPayloads = 0;
+		const int32 BatchSize = FMath::DivideAndRoundUp<int32>(PayloadIds.Num(), NumBatches);
+		NumBatches -= (((NumBatches * BatchSize) - PayloadIds.Num()) / BatchSize);
+
+		check(NumBatches * BatchSize >= PayloadIds.Num());
+
+		for (int32 BatchIndex = 0; BatchIndex < NumBatches; ++BatchIndex)
+		{
+			const int32 BatchStart = BatchIndex * BatchSize;
+			const int32 BatchEnd = FMath::Min((BatchIndex + 1) * BatchSize, PayloadIds.Num());
+
+			TArrayView<const FIoHash> Slice = MakeArrayView(&PayloadIds[BatchStart], BatchEnd - BatchStart);
+			
+			Tasks.Emplace(UE::Tasks::Launch(UE_SOURCE_LOCATION, 
+				[Slice, &System, &NumCompletedPayloads]()
+				{
+					for (const FIoHash& PayloadId : Slice)
+					{
+						FCompressedBuffer Payload = System.PullData(PayloadId);
+						if (Payload.IsNull())
+						{
+							UE_LOG(LogVirtualization, Error, TEXT("Failed to precache payload '%s'"), *LexToString(PayloadId));
+						}
+
+						NumCompletedPayloads++;
+					}
+				}));
+		}	
+
+		while (!WaitOnTasks(Tasks, FTimespan::FromSeconds(30.0)))
+		{
+			const int32 CurrentCompletedPayloads = NumCompletedPayloads;
+
+			const float Progress = ((float)CurrentCompletedPayloads / (float)PayloadIds.Num()) * 100.0f;
+
+			UE_LOG(LogVirtualization, Display, TEXT("Cached %d/%d (%.1f%%)"), CurrentCompletedPayloads, PayloadIds.Num(), Progress);
+		}
+	}
+#elif UE_ITERATION_METHOD == 3
+	{
+		const int NumWorkers = FPlatformMisc::NumberOfCores();
+		const int BatchSize = 64;
+		const int NumPayloads = PayloadIds.Num();
+
+		std::atomic<int32> NumCompletedPayloads = 0;
+
+		FWorkQueue WorkQueue(MoveTemp(PayloadIds), BatchSize);
+
+		TArray<UE::Tasks::FTask> Tasks;
+		Tasks.Reserve(NumWorkers);
+
+		for (int32 Index = 0; Index < NumWorkers; ++Index)
+		{
+			Tasks.Emplace(UE::Tasks::Launch(UE_SOURCE_LOCATION,
+				[&WorkQueue, &System, &NumCompletedPayloads]()
+				{
+					while(true)
+					{
+						FWorkQueue::FJob Job = WorkQueue.GetJob();
+
+						if (Job.IsEmpty())
+						{
+							return;
+						}
+
+						for (const FIoHash& PayloadId : Job)
+						{
+							FCompressedBuffer Payload = System.PullData(PayloadId);
+							if (Payload.IsNull())
+							{
+								UE_LOG(LogVirtualization, Error, TEXT("Failed to precache payload '%s'"), *LexToString(PayloadId));
+							}
+
+							NumCompletedPayloads++;
+						}
+					}
+				}));
+		}
+
+		while (!WaitOnTasks(Tasks, FTimespan::FromSeconds(30.0)))
+		{
+			const int32 CurrentCompletedPayloads = NumCompletedPayloads;
+
+			const float Progress = ((float)CurrentCompletedPayloads / (float)NumPayloads) * 100.0f;
+
+			UE_LOG(LogVirtualization, Display, TEXT("Cached %d/%d (%.1f%%)"), CurrentCompletedPayloads, NumPayloads, Progress);
+		}
+	}
+
+#endif // UE_ITERATION_METHOD
+
+	UE_LOG(LogVirtualization, Display, TEXT("Precaching complete!"));
+
+	UE::Virtualization::IVirtualizationSystem::Get().DumpStats();
+	
+	return  0;
+}
diff --git a/Engine/Source/Editor/VirtualizationEditor/Private/PrecachePayloads.h b/Engine/Source/Editor/VirtualizationEditor/Private/PrecachePayloads.h
new file mode 100644
--- /dev/null
+++ b/Engine/Source/Editor/VirtualizationEditor/Private/PrecachePayloads.h
@@ -0,0 +1,21 @@
+// Copyright Epic Games, Inc. All Rights Reserved.
+
+#pragma once
+
+#include "Commandlets/Commandlet.h"
+#include "UObject/ObjectMacros.h"
+
+#include "PrecachePayloads.generated.h"
+
+UCLASS()
+class UPrecachePayloadsCommandlet
+	: public UCommandlet
+{
+	GENERATED_UCLASS_BODY()
+
+	//~ Begin UCommandlet Interface
+	virtual int32 Main(const FString& Params) override;
+	//~ End UCommandlet Interface
+
+	static int32 StaticMain(const FString& Params);
+};
