Add two new paths to the PreCachePayloads commandlet that take advantage of the new batch API.

#rnx

- Iteration method 4 requests all found payloads in one single large batch. Useful for testing but would easily OOM systems when used on real projects.
- Iteration method 5 is similar to method 3, it splits the work into smaller jobs of 64 payloads at a time and spreads them out over our cores. This is the new method and in tests has managed to reach about 800mbit/s from HordeStorage.

diff --git a/Engine/Source/Editor/VirtualizationEditor/Private/PrecachePayloads.cpp b/Engine/Source/Editor/VirtualizationEditor/Private/PrecachePayloads.cpp
--- a/Engine/Source/Editor/VirtualizationEditor/Private/PrecachePayloads.cpp
+++ b/Engine/Source/Editor/VirtualizationEditor/Private/PrecachePayloads.cpp
@@ -11,8 +11,10 @@
 // 1 - ParallelFor
 // 2 - Split the work into equal sized batches across all cores
 // 3 - Use a limited number of worker tasks
+// 4 - Single batch request 
+// 5 - Use a limited number of worker tasks with batch requests
 
-#define UE_ITERATION_METHOD 3
+#define UE_ITERATION_METHOD 5
 
 namespace 
 {
@@ -128,6 +130,21 @@ TArray<FIoHash> FindVirtualizedPayloads(const TArray<FString>& PackageNames)
 
 }
 
+/** Utility to turn an array of FIoHash into an array of FPullRequest */
+TArray<UE::Virtualization::FPullRequest> ToRequestArray(TConstArrayView<FIoHash> IdentifierArray)
+{
+	TArray<UE::Virtualization::FPullRequest> Requests;
+	Requests.Reserve(IdentifierArray.Num());
+
+	for (const FIoHash& Id : IdentifierArray)
+	{
+		Requests.Emplace(Id);
+	}
+
+	return Requests;
+}
+
+
 UPrecachePayloadsCommandlet::UPrecachePayloadsCommandlet(const FObjectInitializer& ObjectInitializer)
 	: Super(ObjectInitializer)
 {
@@ -148,7 +165,7 @@ int32 UPrecachePayloadsCommandlet::Main(const FString& Params)
 		UE_LOG(LogVirtualization, Display, TEXT("No virtualized payloads found"));
 		return 0;
 	}
-		
+
 	UE_LOG(LogVirtualization, Display, TEXT("Found %d virtualized payloads to precache"), PayloadIds.Num());
 	UE_LOG(LogVirtualization, Display, TEXT("Precaching payloads..."));
 
@@ -273,7 +290,72 @@ int32 UPrecachePayloadsCommandlet::Main(const FString& Params)
 			UE_LOG(LogVirtualization, Display, TEXT("Cached %d/%d (%.1f%%)"), CurrentCompletedPayloads, NumPayloads, Progress);
 		}
 	}
+#elif UE_ITERATION_METHOD == 4
+	{
+		TArray<UE::Virtualization::FPullRequest> Requests = ToRequestArray(PayloadIds);
+	
+		System.PullData(Requests);
+
+		for (const UE::Virtualization::FPullRequest& Request : Requests)
+		{
+			if (!Request.IsSuccess())
+			{
+				UE_LOG(LogVirtualization, Error, TEXT("Failed to precache payload '%s'"), *LexToString(Request.GetIdentifier()));
+			}
+		}
+	}
+#elif UE_ITERATION_METHOD == 5
+	const int NumWorkers = FPlatformMisc::NumberOfCores();
+	const int BatchSize = 64;
+	const int NumPayloads = PayloadIds.Num();
+
+	std::atomic<int32> NumCompletedPayloads = 0;
+
+	FWorkQueue WorkQueue(MoveTemp(PayloadIds), BatchSize);
+
+	TArray<UE::Tasks::FTask> Tasks;
+	Tasks.Reserve(NumWorkers);
+
+	for (int32 Index = 0; Index < NumWorkers; ++Index)
+	{
+		Tasks.Emplace(UE::Tasks::Launch(UE_SOURCE_LOCATION,
+			[&WorkQueue, &System, &NumCompletedPayloads]()
+			{
+				while (true)
+				{
+					FWorkQueue::FJob Job = WorkQueue.GetJob();
+
+					if (Job.IsEmpty())
+					{
+						return;
+					}
+
+					TArray<UE::Virtualization::FPullRequest> Requests = ToRequestArray(Job);
+					if (!System.PullData(Requests))
+					{
+						for (const FIoHash& PayloadId : Job)
+						{
+							FCompressedBuffer Payload = System.PullData(PayloadId);
+							if (Payload.IsNull())
+							{
+								UE_LOG(LogVirtualization, Error, TEXT("Failed to precache payload '%s'"), *LexToString(PayloadId));
+							}
+						}
+					}
+
+					NumCompletedPayloads += Job.Num();
+				}
+			}));
+	}
+
+	while (!WaitOnTasks(Tasks, FTimespan::FromSeconds(30.0)))
+	{
+		const int32 CurrentCompletedPayloads = NumCompletedPayloads;
 
+		const float Progress = ((float)CurrentCompletedPayloads / (float)NumPayloads) * 100.0f;
+
+		UE_LOG(LogVirtualization, Display, TEXT("Cached %d/%d (%.1f%%)"), CurrentCompletedPayloads, NumPayloads, Progress);
+	}
 #endif // UE_ITERATION_METHOD
 
 	UE_LOG(LogVirtualization, Display, TEXT("Precaching complete!"));
