Drastically reduce the amount of time payload precaching spends trying to extract payload information from packages.

#rnx

- Use ParallelForWithTaskContext to iterate ove rthe list of packages and extra the payload information.
- This has been shown to givea speed up of 5x-10x depending on the machine being tested.

diff --git a/Engine/Source/Editor/VirtualizationEditor/Private/PrecachePayloads.cpp b/Engine/Source/Editor/VirtualizationEditor/Private/PrecachePayloads.cpp
--- a/Engine/Source/Editor/VirtualizationEditor/Private/PrecachePayloads.cpp
+++ b/Engine/Source/Editor/VirtualizationEditor/Private/PrecachePayloads.cpp
@@ -113,19 +113,31 @@ TArray<FIoHash> FindVirtualizedPayloads(const TArray<FString>& PackageNames)
 {
 	TRACE_CPUPROFILER_EVENT_SCOPE(FindVirtualizedPayloads);
 
-	TSet<FIoHash> PayloadsToPrecache;
+	// Each task will write out to its own TSet so we don't have to lock anything, we
+	// will combine the sets at the end.
+	TArray<TSet<FIoHash>> PayloadsPerTask;
 
-	for (const FString& PackageName : PackageNames)
-	{
-		UE::FPackageTrailer Trailer;
-		if (UE::FPackageTrailer::TryLoadFromFile(PackageName, Trailer))
+	ParallelForWithTaskContext(PayloadsPerTask, PackageNames.Num(), 
+		[&PackageNames](TSet<FIoHash>& Context, int32 Index)
 		{
-			TArray<FIoHash> VirtualizedPayloads = Trailer.GetPayloads(UE::EPayloadStorageType::Virtualized);
-			PayloadsToPrecache.Append(VirtualizedPayloads);
-		}
+			const FString& PackageName = PackageNames[Index];
+
+			UE::FPackageTrailer Trailer;
+			if (UE::FPackageTrailer::TryLoadFromFile(PackageName, Trailer))
+			{
+				TArray<FIoHash> VirtualizedPayloads = Trailer.GetPayloads(UE::EPayloadStorageType::Virtualized);
+				Context.Append(VirtualizedPayloads);
+			}
+		});
+
+	// Combine the results into a final set
+	TSet<FIoHash> AllPayloads;
+	for (const TSet<FIoHash>& T : PayloadsPerTask)
+	{
+		AllPayloads.Append(T);
 	}
 
-	return PayloadsToPrecache.Array();
+	return AllPayloads.Array();
 }
 
 }
