Add more profiling scopes to the PrecachePayloads commandlet

#rnx

diff --git a/Engine/Source/Editor/VirtualizationEditor/Private/PrecachePayloads.cpp b/Engine/Source/Editor/VirtualizationEditor/Private/PrecachePayloads.cpp
--- a/Engine/Source/Editor/VirtualizationEditor/Private/PrecachePayloads.cpp
+++ b/Engine/Source/Editor/VirtualizationEditor/Private/PrecachePayloads.cpp
@@ -2,10 +2,10 @@
 
 #include "PrecachePayloads.h"
 
+#include "Async/ParallelFor.h"
 #include "PackageHelperFunctions.h"
-#include "Virtualization/VirtualizationSystem.h"
 #include "UObject/PackageTrailer.h"
-#include "Async/ParallelFor.h"
+#include "Virtualization/VirtualizationSystem.h"
 
 // 0 - No threading
 // 1 - ParallelFor
@@ -15,6 +15,7 @@
 // 5 - Use a limited number of worker tasks with batch requests
 
 #define UE_ITERATION_METHOD 5
+#define UE_VALIDATE_PAYLOADS 0
 
 namespace 
 {
@@ -185,6 +186,8 @@ int32 UPrecachePayloadsCommandlet::Main(const FString& Params)
 
 #if UE_ITERATION_METHOD == 0
 	{
+		TRACE_CPUPROFILER_EVENT_SCOPE(Precache_SingleRequests);
+
 		for (const FIoHash& PayloadId : PayloadIds)
 		{
 			FCompressedBuffer Payload = System.PullData(PayloadId);
@@ -196,6 +199,8 @@ int32 UPrecachePayloadsCommandlet::Main(const FString& Params)
 	}
 #elif UE_ITERATION_METHOD == 1
 	{
+		TRACE_CPUPROFILER_EVENT_SCOPE(Precache_ParallelFor);
+
 		ParallelFor(PayloadIds.Num(), [&PayloadIds, &System](int32 Index)
 			{
 				const FIoHash& PayloadId = PayloadIds[Index];
@@ -209,6 +214,8 @@ int32 UPrecachePayloadsCommandlet::Main(const FString& Params)
 	}
 #elif UE_ITERATION_METHOD == 2
 	{
+		TRACE_CPUPROFILER_EVENT_SCOPE(Precache_ThreadedBatches);
+
 		int32 NumBatches = FPlatformMisc::NumberOfCores();
 
 		TArray<UE::Tasks::FTask> Tasks;
@@ -254,6 +261,8 @@ int32 UPrecachePayloadsCommandlet::Main(const FString& Params)
 	}
 #elif UE_ITERATION_METHOD == 3
 	{
+		TRACE_CPUPROFILER_EVENT_SCOPE(Precache_ThreadedWorkerBatches);
+
 		const int NumWorkers = FPlatformMisc::NumberOfCores();
 		const int BatchSize = 64;
 		const int NumPayloads = PayloadIds.Num();
@@ -304,6 +313,8 @@ int32 UPrecachePayloadsCommandlet::Main(const FString& Params)
 	}
 #elif UE_ITERATION_METHOD == 4
 	{
+		TRACE_CPUPROFILER_EVENT_SCOPE(Precache_SingleBatchRequest);
+
 		TArray<UE::Virtualization::FPullRequest> Requests = ToRequestArray(PayloadIds);
 	
 		System.PullData(Requests);
@@ -317,56 +328,66 @@ int32 UPrecachePayloadsCommandlet::Main(const FString& Params)
 		}
 	}
 #elif UE_ITERATION_METHOD == 5
-	const int NumWorkers = FPlatformMisc::NumberOfCores();
-	const int BatchSize = 64;
-	const int NumPayloads = PayloadIds.Num();
+	{
+		TRACE_CPUPROFILER_EVENT_SCOPE(Precache_ThreadedBatchRequest);
+
+		const int NumWorkers = FPlatformMisc::NumberOfCores();
+		const int BatchSize = 64;
+		const int NumPayloads = PayloadIds.Num();
 
-	std::atomic<int32> NumCompletedPayloads = 0;
+		std::atomic<int32> NumCompletedPayloads = 0;
 
-	FWorkQueue WorkQueue(MoveTemp(PayloadIds), BatchSize);
+		FWorkQueue WorkQueue(MoveTemp(PayloadIds), BatchSize);
 
-	TArray<UE::Tasks::FTask> Tasks;
-	Tasks.Reserve(NumWorkers);
+		TArray<UE::Tasks::FTask> Tasks;
+		Tasks.Reserve(NumWorkers);
 
-	for (int32 Index = 0; Index < NumWorkers; ++Index)
-	{
-		Tasks.Emplace(UE::Tasks::Launch(UE_SOURCE_LOCATION,
-			[&WorkQueue, &System, &NumCompletedPayloads]()
-			{
-				while (true)
+		for (int32 Index = 0; Index < NumWorkers; ++Index)
+		{
+			Tasks.Emplace(UE::Tasks::Launch(UE_SOURCE_LOCATION,
+				[&WorkQueue, &System, &NumCompletedPayloads]()
 				{
-					FWorkQueue::FJob Job = WorkQueue.GetJob();
-
-					if (Job.IsEmpty())
+					while (true)
 					{
-						return;
-					}
+						FWorkQueue::FJob Job = WorkQueue.GetJob();
 
-					TArray<UE::Virtualization::FPullRequest> Requests = ToRequestArray(Job);
-					if (!System.PullData(Requests))
-					{
-						for (const FIoHash& PayloadId : Job)
+						if (Job.IsEmpty())
 						{
-							FCompressedBuffer Payload = System.PullData(PayloadId);
-							if (Payload.IsNull())
+							return;
+						}
+
+						TArray<UE::Virtualization::FPullRequest> Requests = ToRequestArray(Job);
+						if (!System.PullData(Requests))
+						{
+							for (const FIoHash& PayloadId : Job)
 							{
-								UE_LOG(LogVirtualization, Error, TEXT("Failed to precache payload '%s'"), *LexToString(PayloadId));
+								FCompressedBuffer Payload = System.PullData(PayloadId);
+								if (Payload.IsNull())
+								{
+									UE_LOG(LogVirtualization, Error, TEXT("%s: Failed to precache payload"), *LexToString(PayloadId));
+								}
+#if UE_VALIDATE_PAYLOADS
+								else if (Payload.GetRawHash() != PayloadId)
+								{
+									UE_LOG(LogVirtualization, Error, TEXT("%s: DAta was corrupted"), *LexToString(PayloadId));
+								}
+#endif //UE_VALIDATE_PAYLOADS
 							}
 						}
-					}
 
-					NumCompletedPayloads += Job.Num();
-				}
-			}));
-	}
+						NumCompletedPayloads += Job.Num();
+					}
+				}));
+		}
 
-	while (!WaitOnTasks(Tasks, FTimespan::FromSeconds(30.0)))
-	{
-		const int32 CurrentCompletedPayloads = NumCompletedPayloads;
+		while (!WaitOnTasks(Tasks, FTimespan::FromSeconds(30.0)))
+		{
+			const int32 CurrentCompletedPayloads = NumCompletedPayloads;
 
-		const float Progress = ((float)CurrentCompletedPayloads / (float)NumPayloads) * 100.0f;
+			const float Progress = ((float)CurrentCompletedPayloads / (float)NumPayloads) * 100.0f;
 
-		UE_LOG(LogVirtualization, Display, TEXT("Cached %d/%d (%.1f%%)"), CurrentCompletedPayloads, NumPayloads, Progress);
+			UE_LOG(LogVirtualization, Display, TEXT("Cached %d/%d (%.1f%%)"), CurrentCompletedPayloads, NumPayloads, Progress);
+		}
 	}
 #endif // UE_ITERATION_METHOD
 
@@ -376,3 +397,6 @@ int32 UPrecachePayloadsCommandlet::Main(const FString& Params)
 	
 	return  0;
 }
+
+#undef UE_VALIDATE_PAYLOADS
+#undef UE_ITERATION_METHOD
