Add a new commandlet GeneratePayloadManifest that will dump all of the payload info about given set of packages.

#rnx

- This doesn't have great potential for production use but can be useful for teams when trying to work out which filters and limits to use when turning on asset virtualization.

diff --git a/Engine/Source/Editor/VirtualizationEditor/Private/GeneratePayloadMaifest.cpp b/Engine/Source/Editor/VirtualizationEditor/Private/GeneratePayloadMaifest.cpp
new file mode 100644
--- /dev/null
+++ b/Engine/Source/Editor/VirtualizationEditor/Private/GeneratePayloadMaifest.cpp
@@ -0,0 +1,119 @@
+// Copyright Epic Games, Inc. All Rights Reserved.
+
+#include "GeneratePayloadManifest.h"
+
+#include "HAL/FileManager.h"
+#include "Misc/PackageName.h"
+#include "Misc/PackagePath.h"
+#include "Misc/Paths.h"
+#include "PackageHelperFunctions.h"
+#include "UObject/PackageTrailer.h"
+
+namespace
+{
+TArray<FString> DiscoverPackages(const FString& Params)
+{
+	TRACE_CPUPROFILER_EVENT_SCOPE(DiscoverPackages);
+	
+	FString PackageDir;
+	if (FParse::Value(*Params, TEXT("PackageDir="), PackageDir) || FParse::Value(*Params, TEXT("PackageFolder="), PackageDir))
+	{
+		TArray<FString> FilesInPackageFolder;
+		FPackageName::FindPackagesInDirectory(FilesInPackageFolder, PackageDir);
+
+		TArray<FString> PackageNames;
+		for (const FString& BasePath : FilesInPackageFolder)
+		{
+			PackageNames.Add(FPaths::CreateStandardFilename(BasePath));
+		}
+
+		return PackageNames;
+	}
+	else
+	{
+		TArray<FString> PackageNames;
+
+		// TODO: In theory we could check the VA settings and skip engine content if they are 
+		// excluded from virtualization. Not convinced it saves enough time to be worth the 
+		// effort.
+		EPackageNormalizationFlags PackageFilter = NORMALIZE_DefaultFlags;
+
+		const FString AssetSearch = TEXT("*") + FPackageName::GetAssetPackageExtension();
+		const FString MapSearch = TEXT("*") + FPackageName::GetMapPackageExtension();
+
+		TArray<FString> Unused;
+		bool bAnyFound = NormalizePackageNames(Unused, PackageNames, AssetSearch, PackageFilter);
+		bAnyFound |= NormalizePackageNames(Unused, PackageNames, MapSearch, PackageFilter);
+
+		return PackageNames;
+	}
+}
+
+} //namespace
+
+UGeneratePayloadManifestCommandlet::UGeneratePayloadManifestCommandlet(const FObjectInitializer& ObjectInitializer)
+	: Super(ObjectInitializer)
+{
+}
+
+int32 UGeneratePayloadManifestCommandlet::Main(const FString& Params)
+{
+	TRACE_CPUPROFILER_EVENT_SCOPE(UGeneratePayloadManifestCommandlet);
+
+	TArray<FString> PackageNames = DiscoverPackages(Params);
+
+	UE_LOG(LogVirtualization, Display, TEXT("Found %d files to look in"), PackageNames.Num());
+
+	int32 PackageTrailerCount = 0;
+
+	const FString CSVPath = FPaths::ProjectSavedDir() / TEXT("payload_manifest.csv");
+	TUniquePtr<FArchive> ManifestFile(IFileManager::Get().CreateFileWriter(*CSVPath));
+	if (!ManifestFile.IsValid())
+	{
+		UE_LOG(LogVirtualization, Error, TEXT("Failed to open '%s' for write"), *CSVPath);
+		return 1;
+	}
+
+	const ANSICHAR* Headings = "Path,PayloadId,SizeOnDisk,UncompressedSize,StorageType, FilterReason\n";
+	ManifestFile->Serialize((void*)Headings, FPlatformString::Strlen(Headings));
+
+	{
+		TRACE_CPUPROFILER_EVENT_SCOPE(ParsePackageTrailers);
+
+		UE_LOG(LogVirtualization, Display, TEXT("Parsing files..."));
+
+		uint32 FilesParsed = 0;
+		for (const FString& PackagePath : PackageNames)
+		{
+			if (FPackageName::IsPackageFilename(PackagePath))
+			{
+				UE::FPackageTrailer Trailer;
+				if (UE::FPackageTrailer::TryLoadFromFile(PackagePath, Trailer))
+				{
+					PackageTrailerCount++;
+
+					Trailer.ForEachPayload([&ManifestFile,&PackagePath](const FIoHash& Id, uint64 SizeOnDisk, uint64 RawSize, UE::EPayloadAccessMode Mode, UE::Virtualization::EPayloadFilterReason Filter)->void
+						{
+							TAnsiStringBuilder<256> LineBuilder;
+							LineBuilder << PackagePath << "," << Id << "," << SizeOnDisk << "," << RawSize << "," << Mode << "," << *LexToString(Filter) <<"\n";
+							ManifestFile->Serialize((void*)LineBuilder.ToString(), LineBuilder.Len());
+						});
+				}
+			}
+
+			if (++FilesParsed % 10000 == 0)
+			{
+				float PercentCompleted = ((float)FilesParsed / (float)PackageNames.Num()) * 100.0f;
+				UE_LOG(LogVirtualization, Display, TEXT("Parsed %d/%d (%.0f%%)"), FilesParsed, PackageNames.Num(), PercentCompleted);
+			}
+			
+		}
+	}
+
+	ManifestFile.Reset(); // Flush file to disk
+
+	UE_LOG(LogVirtualization, Display, TEXT("Found %d package trailers"), PackageTrailerCount);
+	UE_LOG(LogVirtualization, Display, TEXT("Manifest written to '%s'"), *CSVPath);
+
+	return  0;
+}
\ No newline at end of file
diff --git a/Engine/Source/Editor/VirtualizationEditor/Private/GeneratePayloadManifest.h b/Engine/Source/Editor/VirtualizationEditor/Private/GeneratePayloadManifest.h
new file mode 100644
--- /dev/null
+++ b/Engine/Source/Editor/VirtualizationEditor/Private/GeneratePayloadManifest.h
@@ -0,0 +1,31 @@
+// Copyright Epic Games, Inc. All Rights Reserved.
+
+#pragma once
+
+#include "Commandlets/Commandlet.h"
+#include "UObject/ObjectMacros.h"
+
+#include "GeneratePayloadManifest.generated.h"
+
+/**
+ * Creates a csv file containing info about all of the payloads in a set of packages.
+ * By default the commandlet will parse the payloads of all packages in the current
+ * project but this can be overriden with the cmdline switch -PackageDir=XYZ which
+ * will allow the commandlet to parse the payloads of the packages in a given directory.
+ * 
+ * Because the commmandlet is the VirtualizationEditor module it needs to be invoked 
+ * with the command line:
+ * -run=VirtualizationEditor.GeneratePayloadManifest
+ */
+UCLASS()
+class UGeneratePayloadManifestCommandlet
+	: public UCommandlet
+{
+	GENERATED_UCLASS_BODY()
+
+	//~ Begin UCommandlet Interface
+	virtual int32 Main(const FString& Params) override;
+	//~ End UCommandlet Interface
+
+	static int32 StaticMain(const FString& Params);
+};
diff --git a/Engine/Source/Runtime/CoreUObject/Private/UObject/PackageTrailer.cpp b/Engine/Source/Runtime/CoreUObject/Private/UObject/PackageTrailer.cpp
--- a/Engine/Source/Runtime/CoreUObject/Private/UObject/PackageTrailer.cpp
+++ b/Engine/Source/Runtime/CoreUObject/Private/UObject/PackageTrailer.cpp
@@ -715,6 +715,14 @@ bool FPackageTrailer::UpdatePayloadAsVirtualized(const FIoHash& Identifier)
 	}
 }
 
+void FPackageTrailer::ForEachPayload(TFunctionRef<void(const FIoHash&, uint64, uint64, EPayloadAccessMode, UE::Virtualization::EPayloadFilterReason)> Callback) const
+{
+	for (const Private::FLookupTableEntry& Entry : Header.PayloadLookupTable)
+	{
+		Callback(Entry.Identifier, Entry.CompressedSize, Entry.RawSize, Entry.AccessMode, Entry.FilterFlags);
+	}
+}
+
 EPayloadStatus FPackageTrailer::FindPayloadStatus(const FIoHash& Id) const
 {
 	const Private::FLookupTableEntry* Entry = FindEntryInHeader(Header, Id);
diff --git a/Engine/Source/Runtime/CoreUObject/Public/UObject/PackageTrailer.h b/Engine/Source/Runtime/CoreUObject/Public/UObject/PackageTrailer.h
--- a/Engine/Source/Runtime/CoreUObject/Public/UObject/PackageTrailer.h
+++ b/Engine/Source/Runtime/CoreUObject/Public/UObject/PackageTrailer.h
@@ -471,6 +471,11 @@ public:
 	 */
 	[[nodiscard]] bool UpdatePayloadAsVirtualized(const FIoHash& Identifier);
 
+	/**
+	 * Iterates over all payloads in the trailer and invokes the provoided callback on them
+	 */
+	void ForEachPayload(TFunctionRef<void(const FIoHash&, uint64, uint64, EPayloadAccessMode, UE::Virtualization::EPayloadFilterReason)> Callback) const;
+
 	/** Attempt to find the status of the given payload. @See EPayloadStatus */
 	[[nodiscard]] EPayloadStatus FindPayloadStatus(const FIoHash& Id) const;
 
@@ -576,6 +581,30 @@ private:
  */
 [[nodiscard]] COREUOBJECT_API bool FindPayloadsInPackageFile(const FPackagePath& PackagePath, EPayloadStorageType Filter, TArray<FIoHash>& OutPayloadIds);
 
+/** Allow EPayloadAccessMode to work with string builder */
+template <typename CharType>
+inline TStringBuilderBase<CharType>& operator<<(TStringBuilderBase<CharType>& Builder, EPayloadAccessMode Mode)
+{
+	switch (Mode)
+	{
+		case UE::EPayloadAccessMode::Local:
+			Builder << TEXT("Local");
+			break;
+		case UE::EPayloadAccessMode::Referenced:
+			Builder << TEXT("Referenced");
+			break;
+		case UE::EPayloadAccessMode::Virtualized:
+			Builder << TEXT("Virtualized");
+			break;
+		default:
+			Builder << TEXT("Invalid");
+			break;
+	}
+
+	return Builder;
+}
+
+
 } //namespace UE
 
 [[nodiscard]] COREUOBJECT_API FString LexToString(UE::Virtualization::EPayloadFilterReason FilterFlags);
