Move utility function inside the anonymous namespace.

#rnx

diff --git a/Engine/Source/Editor/VirtualizationEditor/Private/PrecachePayloads.cpp b/Engine/Source/Editor/VirtualizationEditor/Private/PrecachePayloads.cpp
--- a/Engine/Source/Editor/VirtualizationEditor/Private/PrecachePayloads.cpp
+++ b/Engine/Source/Editor/VirtualizationEditor/Private/PrecachePayloads.cpp
@@ -142,8 +142,6 @@ TArray<FIoHash> FindVirtualizedPayloads(const TArray<FString>& PackageNames)
 	return AllPayloads.Array();
 }
 
-}
-
 /** Utility to turn an array of FIoHash into an array of FPullRequest */
 TArray<UE::Virtualization::FPullRequest> ToRequestArray(TConstArrayView<FIoHash> IdentifierArray)
 {
@@ -158,6 +156,8 @@ TArray<UE::Virtualization::FPullRequest> ToRequestArray(TConstArrayView<FIoHash>
 	return Requests;
 }
 
+} // namespace
+
 
 UPrecachePayloadsCommandlet::UPrecachePayloadsCommandlet(const FObjectInitializer& ObjectInitializer)
 	: Super(ObjectInitializer)
