Move some common commandlet code to a utility file so that the virtualization commandlets can share.

#rnx

diff --git a/Engine/Source/Editor/VirtualizationEditor/Private/CommandletUtils.cpp b/Engine/Source/Editor/VirtualizationEditor/Private/CommandletUtils.cpp
new file mode 100644
--- /dev/null
+++ b/Engine/Source/Editor/VirtualizationEditor/Private/CommandletUtils.cpp
@@ -0,0 +1,91 @@
+// Copyright Epic Games, Inc. All Rights Reserved.
+
+#include "CommandletUtils.h"
+
+#include "Async/ParallelFor.h"
+#include "Containers/Set.h"
+#include "Misc/PackageName.h"
+#include "PackageHelperFunctions.h"
+#include "UObject/PackageTrailer.h"
+
+namespace UE::Virtualization
+{
+
+TArray<FString> FindAllPackages()
+{
+	TRACE_CPUPROFILER_EVENT_SCOPE(FindAllPackages);
+
+	// TODO: In theory we could check the VA settings and skip engine content if they are 
+	// excluded from virtualization. Not convinced it saves enough time to be worth the 
+	// effort.
+	const EPackageNormalizationFlags PackageFilter = NORMALIZE_DefaultFlags;
+
+	const FString AssetSearch = TEXT("*") + FPackageName::GetAssetPackageExtension();
+	const FString MapSearch = TEXT("*") + FPackageName::GetMapPackageExtension();
+
+	TArray<FString> PackageNames;
+	TArray<FString> Unused;
+	bool bAnyFound = NormalizePackageNames(Unused, PackageNames, AssetSearch, PackageFilter);
+	bAnyFound |= NormalizePackageNames(Unused, PackageNames, MapSearch, PackageFilter);
+
+	return PackageNames;
+}
+
+TArray<FString> DiscoverPackages(const FString& CmdlineParams)
+{
+	TRACE_CPUPROFILER_EVENT_SCOPE(DiscoverPackages);
+
+	FString PackageDir;
+	if (FParse::Value(*CmdlineParams, TEXT("PackageDir="), PackageDir) || FParse::Value(*CmdlineParams, TEXT("PackageFolder="), PackageDir))
+	{
+		TArray<FString> FilesInPackageFolder;
+		FPackageName::FindPackagesInDirectory(FilesInPackageFolder, PackageDir);
+
+		TArray<FString> PackageNames;
+		for (const FString& BasePath : FilesInPackageFolder)
+		{
+			PackageNames.Add(FPaths::CreateStandardFilename(BasePath));
+		}
+
+		return PackageNames;
+	}
+	else
+	{
+		return FindAllPackages();
+	}
+}
+
+
+
+TArray<FIoHash> FindVirtualizedPayloads(const TArray<FString>& PackagePaths)
+{
+	TRACE_CPUPROFILER_EVENT_SCOPE(FindVirtualizedPayloads);
+
+	// Each task will write out to its own TSet so we don't have to lock anything, we
+	// will combine the sets at the end.
+	TArray<TSet<FIoHash>> PayloadsPerTask;
+
+	ParallelForWithTaskContext(PayloadsPerTask, PackagePaths.Num(),
+		[&PackagePaths](TSet<FIoHash>& Context, int32 Index)
+		{
+			const FString& PackageName = PackagePaths[Index];
+
+			UE::FPackageTrailer Trailer;
+			if (UE::FPackageTrailer::TryLoadFromFile(PackageName, Trailer))
+			{
+				TArray<FIoHash> VirtualizedPayloads = Trailer.GetPayloads(UE::EPayloadStorageType::Virtualized);
+				Context.Append(VirtualizedPayloads);
+			}
+		});
+
+	// Combine the results into a final set
+	TSet<FIoHash> AllPayloads;
+	for (const TSet<FIoHash>& Payloads : PayloadsPerTask)
+	{
+		AllPayloads.Append(Payloads);
+	}
+
+	return AllPayloads.Array();
+}
+
+} //namespace UE::Virtualization
diff --git a/Engine/Source/Editor/VirtualizationEditor/Private/CommandletUtils.h b/Engine/Source/Editor/VirtualizationEditor/Private/CommandletUtils.h
new file mode 100644
--- /dev/null
+++ b/Engine/Source/Editor/VirtualizationEditor/Private/CommandletUtils.h
@@ -0,0 +1,31 @@
+// Copyright Epic Games, Inc. All Rights Reserved.
+
+#pragma once
+
+#include "Containers/Array.h"
+#include "Containers/UnrealString.h"
+#include "IO/IoHash.h"
+
+namespace UE::Virtualization
+{
+
+/** Parse all of the active mount points and find all .uasset/.umaps */
+TArray<FString> FindAllPackages();
+
+/**
+ * Finds all of the packages under a the directory given by the provided command line.
+ * If no commandline switch can be found then the function will return all avaliable 
+ * packages.
+ * Valid commandline switches:
+ * '-PackageDir=...'
+ * '-PackageFolder=...'
+ * 
+ * @param CmdlineParams A string containing the command line
+ * @return An array with the file path for each package found
+ */
+TArray<FString> DiscoverPackages(const FString& CmdlineParams);
+
+/** Returns a combined list of unique virtualized payload ids from the given list of packages */
+TArray<FIoHash> FindVirtualizedPayloads(const TArray<FString>& PackageNames);
+
+} //namespace UE::Virtualization
diff --git a/Engine/Source/Editor/VirtualizationEditor/Private/GeneratePayloadMaifest.cpp b/Engine/Source/Editor/VirtualizationEditor/Private/GeneratePayloadMaifest.cpp
--- a/Engine/Source/Editor/VirtualizationEditor/Private/GeneratePayloadMaifest.cpp
+++ b/Engine/Source/Editor/VirtualizationEditor/Private/GeneratePayloadMaifest.cpp
@@ -2,55 +2,13 @@
 
 #include "GeneratePayloadManifest.h"
 
+#include "CommandletUtils.h"
 #include "HAL/FileManager.h"
 #include "Misc/PackageName.h"
 #include "Misc/PackagePath.h"
 #include "Misc/Paths.h"
-#include "PackageHelperFunctions.h"
 #include "UObject/PackageTrailer.h"
 
-namespace
-{
-TArray<FString> DiscoverPackages(const FString& Params)
-{
-	TRACE_CPUPROFILER_EVENT_SCOPE(DiscoverPackages);
-	
-	FString PackageDir;
-	if (FParse::Value(*Params, TEXT("PackageDir="), PackageDir) || FParse::Value(*Params, TEXT("PackageFolder="), PackageDir))
-	{
-		TArray<FString> FilesInPackageFolder;
-		FPackageName::FindPackagesInDirectory(FilesInPackageFolder, PackageDir);
-
-		TArray<FString> PackageNames;
-		for (const FString& BasePath : FilesInPackageFolder)
-		{
-			PackageNames.Add(FPaths::CreateStandardFilename(BasePath));
-		}
-
-		return PackageNames;
-	}
-	else
-	{
-		TArray<FString> PackageNames;
-
-		// TODO: In theory we could check the VA settings and skip engine content if they are 
-		// excluded from virtualization. Not convinced it saves enough time to be worth the 
-		// effort.
-		EPackageNormalizationFlags PackageFilter = NORMALIZE_DefaultFlags;
-
-		const FString AssetSearch = TEXT("*") + FPackageName::GetAssetPackageExtension();
-		const FString MapSearch = TEXT("*") + FPackageName::GetMapPackageExtension();
-
-		TArray<FString> Unused;
-		bool bAnyFound = NormalizePackageNames(Unused, PackageNames, AssetSearch, PackageFilter);
-		bAnyFound |= NormalizePackageNames(Unused, PackageNames, MapSearch, PackageFilter);
-
-		return PackageNames;
-	}
-}
-
-} //namespace
-
 UGeneratePayloadManifestCommandlet::UGeneratePayloadManifestCommandlet(const FObjectInitializer& ObjectInitializer)
 	: Super(ObjectInitializer)
 {
@@ -60,7 +18,7 @@ int32 UGeneratePayloadManifestCommandlet::Main(const FString& Params)
 {
 	TRACE_CPUPROFILER_EVENT_SCOPE(UGeneratePayloadManifestCommandlet);
 
-	TArray<FString> PackageNames = DiscoverPackages(Params);
+	TArray<FString> PackageNames = UE::Virtualization::DiscoverPackages(Params);
 
 	UE_LOG(LogVirtualization, Display, TEXT("Found %d files to look in"), PackageNames.Num());
 
diff --git a/Engine/Source/Editor/VirtualizationEditor/Private/PrecachePayloads.cpp b/Engine/Source/Editor/VirtualizationEditor/Private/PrecachePayloads.cpp
--- a/Engine/Source/Editor/VirtualizationEditor/Private/PrecachePayloads.cpp
+++ b/Engine/Source/Editor/VirtualizationEditor/Private/PrecachePayloads.cpp
@@ -3,7 +3,7 @@
 #include "PrecachePayloads.h"
 
 #include "Async/ParallelFor.h"
-#include "PackageHelperFunctions.h"
+#include "CommandletUtils.h"
 #include "Tasks/Task.h"
 #include "UObject/PackageTrailer.h"
 #include "Virtualization/VirtualizationSystem.h"
@@ -88,60 +88,6 @@ bool WaitOnTasks(const TArray<UE::Tasks::FTask>& Tasks, FTimespan InTimeout)
 	return CompletionEvent->Wait((uint32)FMath::Clamp<int64>(Timeout.GetRemainingTime().GetTicks() / ETimespan::TicksPerMillisecond, 0, MAX_uint32));
 }
 
-/** Parse all of the active mount points and find all .uasset/.umaps */
-TArray<FString> FindPackages()
-{
-	TRACE_CPUPROFILER_EVENT_SCOPE(FindPackages);
-
-	TArray<FString> PackageNames;
-
-	// TODO: In theory we could check the VA settings and skip engine content if they are 
-	// excluded from virtualization. Not convinced it saves enough time to be worth the 
-	// effort.
-	EPackageNormalizationFlags PackageFilter = NORMALIZE_DefaultFlags;
-
-	const FString AssetSearch = TEXT("*") + FPackageName::GetAssetPackageExtension();
-	const FString MapSearch = TEXT("*") + FPackageName::GetMapPackageExtension();
-	
-	TArray<FString> Unused;
-	bool bAnyFound = NormalizePackageNames(Unused, PackageNames, AssetSearch, PackageFilter);
-	bAnyFound |= NormalizePackageNames(Unused, PackageNames, MapSearch, PackageFilter);
-
-	return PackageNames;
-}
-
-/** Returns a combined list of unique virtualized payload ids from the given list of packages */
-TArray<FIoHash> FindVirtualizedPayloads(const TArray<FString>& PackageNames)
-{
-	TRACE_CPUPROFILER_EVENT_SCOPE(FindVirtualizedPayloads);
-
-	// Each task will write out to its own TSet so we don't have to lock anything, we
-	// will combine the sets at the end.
-	TArray<TSet<FIoHash>> PayloadsPerTask;
-
-	ParallelForWithTaskContext(PayloadsPerTask, PackageNames.Num(), 
-		[&PackageNames](TSet<FIoHash>& Context, int32 Index)
-		{
-			const FString& PackageName = PackageNames[Index];
-
-			UE::FPackageTrailer Trailer;
-			if (UE::FPackageTrailer::TryLoadFromFile(PackageName, Trailer))
-			{
-				TArray<FIoHash> VirtualizedPayloads = Trailer.GetPayloads(UE::EPayloadStorageType::Virtualized);
-				Context.Append(VirtualizedPayloads);
-			}
-		});
-
-	// Combine the results into a final set
-	TSet<FIoHash> AllPayloads;
-	for (const TSet<FIoHash>& T : PayloadsPerTask)
-	{
-		AllPayloads.Append(T);
-	}
-
-	return AllPayloads.Array();
-}
-
 /** Utility to turn an array of FIoHash into an array of FPullRequest */
 TArray<UE::Virtualization::FPullRequest> ToRequestArray(TConstArrayView<FIoHash> IdentifierArray)
 {
@@ -168,11 +114,11 @@ int32 UPrecachePayloadsCommandlet::Main(const FString& Params)
 {
 	TRACE_CPUPROFILER_EVENT_SCOPE(UPrecachePayloadsCommandlet);
 
-	TArray<FString> PackageNames = FindPackages();
+	TArray<FString> PackageNames = UE::Virtualization::FindAllPackages();
 
 	UE_LOG(LogVirtualization, Display, TEXT("Found %d packages"), PackageNames.Num());
 
-	TArray<FIoHash> PayloadIds = FindVirtualizedPayloads(PackageNames);
+	TArray<FIoHash> PayloadIds = UE::Virtualization::FindVirtualizedPayloads(PackageNames);
 
 	if (PayloadIds.IsEmpty())
 	{
