Rename virtualization commandlet files to a more standard naming convention.

#rnx

diff --git a/Engine/Source/Editor/VirtualizationEditor/Private/GeneratePayloadMaifest.cpp b/Engine/Source/Editor/VirtualizationEditor/Private/GeneratePayloadMaifest.cpp
deleted file mode 100644
diff --git a/Engine/Source/Editor/VirtualizationEditor/Private/GeneratePayloadManifest.h b/Engine/Source/Editor/VirtualizationEditor/Private/GeneratePayloadManifest.h
deleted file mode 100644
diff --git a/Engine/Source/Editor/VirtualizationEditor/Private/GeneratePayloadManifestCommandlet.cpp b/Engine/Source/Editor/VirtualizationEditor/Private/GeneratePayloadManifestCommandlet.cpp
new file mode 100644
--- /dev/null
+++ b/Engine/Source/Editor/VirtualizationEditor/Private/GeneratePayloadManifestCommandlet.cpp
@@ -0,0 +1,77 @@
+// Copyright Epic Games, Inc. All Rights Reserved.
+
+#include "GeneratePayloadManifestCommandlet.h"
+
+#include "CommandletUtils.h"
+#include "HAL/FileManager.h"
+#include "Misc/PackageName.h"
+#include "Misc/PackagePath.h"
+#include "Misc/Paths.h"
+#include "UObject/PackageTrailer.h"
+
+UGeneratePayloadManifestCommandlet::UGeneratePayloadManifestCommandlet(const FObjectInitializer& ObjectInitializer)
+	: Super(ObjectInitializer)
+{
+}
+
+int32 UGeneratePayloadManifestCommandlet::Main(const FString& Params)
+{
+	TRACE_CPUPROFILER_EVENT_SCOPE(UGeneratePayloadManifestCommandlet);
+
+	TArray<FString> PackageNames = UE::Virtualization::DiscoverPackages(Params);
+
+	UE_LOG(LogVirtualization, Display, TEXT("Found %d files to look in"), PackageNames.Num());
+
+	int32 PackageTrailerCount = 0;
+
+	const FString CSVPath = FPaths::ProjectSavedDir() / TEXT("payload_manifest.csv");
+	TUniquePtr<FArchive> ManifestFile(IFileManager::Get().CreateFileWriter(*CSVPath));
+	if (!ManifestFile.IsValid())
+	{
+		UE_LOG(LogVirtualization, Error, TEXT("Failed to open '%s' for write"), *CSVPath);
+		return 1;
+	}
+
+	const ANSICHAR* Headings = "Path,PayloadId,SizeOnDisk,UncompressedSize,StorageType, FilterReason\n";
+	ManifestFile->Serialize((void*)Headings, FPlatformString::Strlen(Headings));
+
+	{
+		TRACE_CPUPROFILER_EVENT_SCOPE(ParsePackageTrailers);
+
+		UE_LOG(LogVirtualization, Display, TEXT("Parsing files..."));
+
+		uint32 FilesParsed = 0;
+		for (const FString& PackagePath : PackageNames)
+		{
+			if (FPackageName::IsPackageFilename(PackagePath))
+			{
+				UE::FPackageTrailer Trailer;
+				if (UE::FPackageTrailer::TryLoadFromFile(PackagePath, Trailer))
+				{
+					PackageTrailerCount++;
+
+					Trailer.ForEachPayload([&ManifestFile,&PackagePath](const FIoHash& Id, uint64 SizeOnDisk, uint64 RawSize, UE::EPayloadAccessMode Mode, UE::Virtualization::EPayloadFilterReason Filter)->void
+						{
+							TAnsiStringBuilder<256> LineBuilder;
+							LineBuilder << PackagePath << "," << Id << "," << SizeOnDisk << "," << RawSize << "," << Mode << "," << *LexToString(Filter) <<"\n";
+							ManifestFile->Serialize((void*)LineBuilder.ToString(), LineBuilder.Len());
+						});
+				}
+			}
+
+			if (++FilesParsed % 10000 == 0)
+			{
+				float PercentCompleted = ((float)FilesParsed / (float)PackageNames.Num()) * 100.0f;
+				UE_LOG(LogVirtualization, Display, TEXT("Parsed %d/%d (%.0f%%)"), FilesParsed, PackageNames.Num(), PercentCompleted);
+			}
+			
+		}
+	}
+
+	ManifestFile.Reset(); // Flush file to disk
+
+	UE_LOG(LogVirtualization, Display, TEXT("Found %d package trailers"), PackageTrailerCount);
+	UE_LOG(LogVirtualization, Display, TEXT("Manifest written to '%s'"), *CSVPath);
+
+	return  0;
+}
\ No newline at end of file
diff --git a/Engine/Source/Editor/VirtualizationEditor/Private/GeneratePayloadManifestCommandlet.h b/Engine/Source/Editor/VirtualizationEditor/Private/GeneratePayloadManifestCommandlet.h
new file mode 100644
--- /dev/null
+++ b/Engine/Source/Editor/VirtualizationEditor/Private/GeneratePayloadManifestCommandlet.h
@@ -0,0 +1,31 @@
+// Copyright Epic Games, Inc. All Rights Reserved.
+
+#pragma once
+
+#include "Commandlets/Commandlet.h"
+#include "UObject/ObjectMacros.h"
+
+#include "GeneratePayloadManifestCommandlet.generated.h"
+
+/**
+ * Creates a csv file containing info about all of the payloads in a set of packages.
+ * By default the commandlet will parse the payloads of all packages in the current
+ * project but this can be overriden with the cmdline switch -PackageDir=XYZ which
+ * will allow the commandlet to parse the payloads of the packages in a given directory.
+ * 
+ * Because the commmandlet is the VirtualizationEditor module it needs to be invoked 
+ * with the command line:
+ * -run=VirtualizationEditor.GeneratePayloadManifest
+ */
+UCLASS()
+class UGeneratePayloadManifestCommandlet
+	: public UCommandlet
+{
+	GENERATED_UCLASS_BODY()
+
+	//~ Begin UCommandlet Interface
+	virtual int32 Main(const FString& Params) override;
+	//~ End UCommandlet Interface
+
+	static int32 StaticMain(const FString& Params);
+};
diff --git a/Engine/Source/Editor/VirtualizationEditor/Private/PrecachePayloads.cpp b/Engine/Source/Editor/VirtualizationEditor/Private/PrecachePayloads.cpp
deleted file mode 100644
diff --git a/Engine/Source/Editor/VirtualizationEditor/Private/PrecachePayloads.h b/Engine/Source/Editor/VirtualizationEditor/Private/PrecachePayloads.h
deleted file mode 100644
diff --git a/Engine/Source/Editor/VirtualizationEditor/Private/PrecachePayloadsCommandlet.cpp b/Engine/Source/Editor/VirtualizationEditor/Private/PrecachePayloadsCommandlet.cpp
new file mode 100644
--- /dev/null
+++ b/Engine/Source/Editor/VirtualizationEditor/Private/PrecachePayloadsCommandlet.cpp
@@ -0,0 +1,349 @@
+// Copyright Epic Games, Inc. All Rights Reserved.
+
+#include "PrecachePayloadsCommandlet.h"
+
+#include "Async/ParallelFor.h"
+#include "CommandletUtils.h"
+#include "Tasks/Task.h"
+#include "UObject/PackageTrailer.h"
+#include "Virtualization/VirtualizationSystem.h"
+
+// 0 - No threading
+// 1 - ParallelFor
+// 2 - Split the work into equal sized batches across all cores
+// 3 - Use a limited number of worker tasks
+// 4 - Single batch request 
+// 5 - Use a limited number of worker tasks with batch requests
+
+#define UE_ITERATION_METHOD 5
+#define UE_VALIDATE_PAYLOADS 0
+
+namespace 
+{
+
+class FWorkQueue
+{
+public:
+	using FJob = TArrayView<const FIoHash>;
+
+	FWorkQueue(TArray<FIoHash>&& InWork, int32 JobSize)
+		: Work(InWork)
+	{
+		const int32 NumJobs = FMath::DivideAndRoundUp<int32>(Work.Num(), JobSize);
+		Jobs.Reserve(NumJobs);
+
+		for (int32 JobIndex = 0; JobIndex < NumJobs; ++JobIndex)
+		{
+			const int32 JobStart = JobIndex * JobSize;
+			const int32 JobEnd = FMath::Min((JobIndex + 1) * JobSize, Work.Num());
+
+			FJob Job = MakeArrayView(&Work[JobStart], JobEnd - JobStart);
+			Jobs.Add(Job);
+		}
+	}
+
+	FJob GetJob()
+	{
+		TRACE_CPUPROFILER_EVENT_SCOPE(FWorkQueue::GetJob);
+
+		FScopeLock _(&JobsCS);
+
+		if (!Jobs.IsEmpty())
+		{
+			return Jobs.Pop(false);
+		}
+		else
+		{
+			return FJob();
+		}	
+	}
+
+
+	TArray<FIoHash> Work;
+	TArray<FJob> Jobs;
+
+	FCriticalSection JobsCS;
+};
+
+/** 
+ * Utility to wait for the given tasks to complete and return early if the tasks are not completed
+ * within the given timespan.
+ * This is used isntead of UE::Tasks::Wait, because that call can execute the remaining tasks if not
+ * already being processed, which in most scenarios makes sense, but since we are running in a commandlet
+ * and we know nothing else is running, we don't need the game thread processing work as well. 
+ */
+bool WaitOnTasks(const TArray<UE::Tasks::FTask>& Tasks, FTimespan InTimeout)
+{
+	UE::FTimeout Timeout{ InTimeout };
+
+	FSharedEventRef CompletionEvent;
+
+	UE::Tasks::FTask WaitingTask = UE::Tasks::Launch(UE_SOURCE_LOCATION,
+		[CompletionEvent]
+		{
+			CompletionEvent->Trigger();
+		},
+		Tasks, LowLevelTasks::ETaskPriority::Default, UE::Tasks::EExtendedTaskPriority::Inline);
+
+	return CompletionEvent->Wait((uint32)FMath::Clamp<int64>(Timeout.GetRemainingTime().GetTicks() / ETimespan::TicksPerMillisecond, 0, MAX_uint32));
+}
+
+/** Utility to turn an array of FIoHash into an array of FPullRequest */
+TArray<UE::Virtualization::FPullRequest> ToRequestArray(TConstArrayView<FIoHash> IdentifierArray)
+{
+	TArray<UE::Virtualization::FPullRequest> Requests;
+	Requests.Reserve(IdentifierArray.Num());
+
+	for (const FIoHash& Id : IdentifierArray)
+	{
+		Requests.Emplace(Id);
+	}
+
+	return Requests;
+}
+
+} // namespace
+
+
+UPrecachePayloadsCommandlet::UPrecachePayloadsCommandlet(const FObjectInitializer& ObjectInitializer)
+	: Super(ObjectInitializer)
+{
+}
+
+int32 UPrecachePayloadsCommandlet::Main(const FString& Params)
+{
+	TRACE_CPUPROFILER_EVENT_SCOPE(UPrecachePayloadsCommandlet);
+
+	TArray<FString> PackageNames = UE::Virtualization::FindAllPackages();
+
+	UE_LOG(LogVirtualization, Display, TEXT("Found %d packages"), PackageNames.Num());
+
+	TArray<FIoHash> PayloadIds = UE::Virtualization::FindVirtualizedPayloads(PackageNames);
+
+	if (PayloadIds.IsEmpty())
+	{
+		UE_LOG(LogVirtualization, Display, TEXT("No virtualized payloads found"));
+		return 0;
+	}
+
+	UE_LOG(LogVirtualization, Display, TEXT("Found %d virtualized payloads to precache"), PayloadIds.Num());
+	UE_LOG(LogVirtualization, Display, TEXT("Precaching payloads..."));
+
+	UE::Virtualization::IVirtualizationSystem& System = UE::Virtualization::IVirtualizationSystem::Get();
+
+#if UE_ITERATION_METHOD == 0
+	{
+		TRACE_CPUPROFILER_EVENT_SCOPE(Precache_SingleRequests);
+
+		for (const FIoHash& PayloadId : PayloadIds)
+		{
+			FCompressedBuffer Payload = System.PullData(PayloadId);
+			if (Payload.IsNull())
+			{
+				UE_LOG(LogVirtualization, Error, TEXT("Failed to precache payload '%s'"), *LexToString(PayloadId));
+			}
+		}
+	}
+#elif UE_ITERATION_METHOD == 1
+	{
+		TRACE_CPUPROFILER_EVENT_SCOPE(Precache_ParallelFor);
+
+		ParallelFor(PayloadIds.Num(), [&PayloadIds, &System](int32 Index)
+			{
+				const FIoHash& PayloadId = PayloadIds[Index];
+
+				FCompressedBuffer Payload = System.PullData(PayloadId);
+				if (Payload.IsNull())
+				{
+					UE_LOG(LogVirtualization, Error, TEXT("Failed to precache payload '%s'"), *LexToString(PayloadId));
+				}
+			});
+	}
+#elif UE_ITERATION_METHOD == 2
+	{
+		TRACE_CPUPROFILER_EVENT_SCOPE(Precache_ThreadedBatches);
+
+		int32 NumBatches = FPlatformMisc::NumberOfCores();
+
+		TArray<UE::Tasks::FTask> Tasks;
+		Tasks.Reserve(NumBatches);
+		
+		std::atomic<int32> NumCompletedPayloads = 0;
+		const int32 BatchSize = FMath::DivideAndRoundUp<int32>(PayloadIds.Num(), NumBatches);
+		NumBatches -= (((NumBatches * BatchSize) - PayloadIds.Num()) / BatchSize);
+
+		check(NumBatches * BatchSize >= PayloadIds.Num());
+
+		for (int32 BatchIndex = 0; BatchIndex < NumBatches; ++BatchIndex)
+		{
+			const int32 BatchStart = BatchIndex * BatchSize;
+			const int32 BatchEnd = FMath::Min((BatchIndex + 1) * BatchSize, PayloadIds.Num());
+
+			TArrayView<const FIoHash> Slice = MakeArrayView(&PayloadIds[BatchStart], BatchEnd - BatchStart);
+			
+			Tasks.Emplace(UE::Tasks::Launch(UE_SOURCE_LOCATION, 
+				[Slice, &System, &NumCompletedPayloads]()
+				{
+					for (const FIoHash& PayloadId : Slice)
+					{
+						FCompressedBuffer Payload = System.PullData(PayloadId);
+						if (Payload.IsNull())
+						{
+							UE_LOG(LogVirtualization, Error, TEXT("Failed to precache payload '%s'"), *LexToString(PayloadId));
+						}
+
+						NumCompletedPayloads++;
+					}
+				}));
+		}	
+
+		while (!WaitOnTasks(Tasks, FTimespan::FromSeconds(30.0)))
+		{
+			const int32 CurrentCompletedPayloads = NumCompletedPayloads;
+
+			const float Progress = ((float)CurrentCompletedPayloads / (float)PayloadIds.Num()) * 100.0f;
+
+			UE_LOG(LogVirtualization, Display, TEXT("Cached %d/%d (%.1f%%)"), CurrentCompletedPayloads, PayloadIds.Num(), Progress);
+		}
+	}
+#elif UE_ITERATION_METHOD == 3
+	{
+		TRACE_CPUPROFILER_EVENT_SCOPE(Precache_ThreadedWorkerBatches);
+
+		const int NumWorkers = FPlatformMisc::NumberOfCores();
+		const int BatchSize = 64;
+		const int NumPayloads = PayloadIds.Num();
+
+		std::atomic<int32> NumCompletedPayloads = 0;
+
+		FWorkQueue WorkQueue(MoveTemp(PayloadIds), BatchSize);
+
+		TArray<UE::Tasks::FTask> Tasks;
+		Tasks.Reserve(NumWorkers);
+
+		for (int32 Index = 0; Index < NumWorkers; ++Index)
+		{
+			Tasks.Emplace(UE::Tasks::Launch(UE_SOURCE_LOCATION,
+				[&WorkQueue, &System, &NumCompletedPayloads]()
+				{
+					while(true)
+					{
+						FWorkQueue::FJob Job = WorkQueue.GetJob();
+
+						if (Job.IsEmpty())
+						{
+							return;
+						}
+
+						for (const FIoHash& PayloadId : Job)
+						{
+							FCompressedBuffer Payload = System.PullData(PayloadId);
+							if (Payload.IsNull())
+							{
+								UE_LOG(LogVirtualization, Error, TEXT("Failed to precache payload '%s'"), *LexToString(PayloadId));
+							}
+
+							NumCompletedPayloads++;
+						}
+					}
+				}));
+		}
+
+		while (!WaitOnTasks(Tasks, FTimespan::FromSeconds(30.0)))
+		{
+			const int32 CurrentCompletedPayloads = NumCompletedPayloads;
+
+			const float Progress = ((float)CurrentCompletedPayloads / (float)NumPayloads) * 100.0f;
+
+			UE_LOG(LogVirtualization, Display, TEXT("Cached %d/%d (%.1f%%)"), CurrentCompletedPayloads, NumPayloads, Progress);
+		}
+	}
+#elif UE_ITERATION_METHOD == 4
+	{
+		TRACE_CPUPROFILER_EVENT_SCOPE(Precache_SingleBatchRequest);
+
+		TArray<UE::Virtualization::FPullRequest> Requests = ToRequestArray(PayloadIds);
+	
+		System.PullData(Requests);
+
+		for (const UE::Virtualization::FPullRequest& Request : Requests)
+		{
+			if (!Request.IsSuccess())
+			{
+				UE_LOG(LogVirtualization, Error, TEXT("Failed to precache payload '%s'"), *LexToString(Request.GetIdentifier()));
+			}
+		}
+	}
+#elif UE_ITERATION_METHOD == 5
+	{
+		TRACE_CPUPROFILER_EVENT_SCOPE(Precache_ThreadedBatchRequest);
+
+		const int NumWorkers = FPlatformMisc::NumberOfCores();
+		const int BatchSize = 64;
+		const int NumPayloads = PayloadIds.Num();
+
+		std::atomic<int32> NumCompletedPayloads = 0;
+
+		FWorkQueue WorkQueue(MoveTemp(PayloadIds), BatchSize);
+
+		TArray<UE::Tasks::FTask> Tasks;
+		Tasks.Reserve(NumWorkers);
+
+		for (int32 Index = 0; Index < NumWorkers; ++Index)
+		{
+			Tasks.Emplace(UE::Tasks::Launch(UE_SOURCE_LOCATION,
+				[&WorkQueue, &System, &NumCompletedPayloads]()
+				{
+					while (true)
+					{
+						FWorkQueue::FJob Job = WorkQueue.GetJob();
+
+						if (Job.IsEmpty())
+						{
+							return;
+						}
+
+						TArray<UE::Virtualization::FPullRequest> Requests = ToRequestArray(Job);
+						if (!System.PullData(Requests))
+						{
+							for (const FIoHash& PayloadId : Job)
+							{
+								FCompressedBuffer Payload = System.PullData(PayloadId);
+								if (Payload.IsNull())
+								{
+									UE_LOG(LogVirtualization, Error, TEXT("%s: Failed to precache payload"), *LexToString(PayloadId));
+								}
+#if UE_VALIDATE_PAYLOADS
+								else if (Payload.GetRawHash() != PayloadId)
+								{
+									UE_LOG(LogVirtualization, Error, TEXT("%s: DAta was corrupted"), *LexToString(PayloadId));
+								}
+#endif //UE_VALIDATE_PAYLOADS
+							}
+						}
+
+						NumCompletedPayloads += Job.Num();
+					}
+				}));
+		}
+
+		while (!WaitOnTasks(Tasks, FTimespan::FromSeconds(30.0)))
+		{
+			const int32 CurrentCompletedPayloads = NumCompletedPayloads;
+
+			const float Progress = ((float)CurrentCompletedPayloads / (float)NumPayloads) * 100.0f;
+
+			UE_LOG(LogVirtualization, Display, TEXT("Cached %d/%d (%.1f%%)"), CurrentCompletedPayloads, NumPayloads, Progress);
+		}
+	}
+#endif // UE_ITERATION_METHOD
+
+	UE_LOG(LogVirtualization, Display, TEXT("Precaching complete!"));
+
+	UE::Virtualization::IVirtualizationSystem::Get().DumpStats();
+	
+	return  0;
+}
+
+#undef UE_VALIDATE_PAYLOADS
+#undef UE_ITERATION_METHOD
diff --git a/Engine/Source/Editor/VirtualizationEditor/Private/PrecachePayloadsCommandlet.h b/Engine/Source/Editor/VirtualizationEditor/Private/PrecachePayloadsCommandlet.h
new file mode 100644
--- /dev/null
+++ b/Engine/Source/Editor/VirtualizationEditor/Private/PrecachePayloadsCommandlet.h
@@ -0,0 +1,21 @@
+// Copyright Epic Games, Inc. All Rights Reserved.
+
+#pragma once
+
+#include "Commandlets/Commandlet.h"
+#include "UObject/ObjectMacros.h"
+
+#include "PrecachePayloadsCommandlet.generated.h"
+
+UCLASS()
+class UPrecachePayloadsCommandlet
+	: public UCommandlet
+{
+	GENERATED_UCLASS_BODY()
+
+	//~ Begin UCommandlet Interface
+	virtual int32 Main(const FString& Params) override;
+	//~ End UCommandlet Interface
+
+	static int32 StaticMain(const FString& Params);
+};
