Add a new commandlet for virtualizing an entire project.

#rnx

- Cmdline to use '-run=VirtualizationEditor.VirtualizeProject'.
- The commandlet will first scan the project to find all packages (.uasset/.umap) and then attempt to virtualize them.
- If revision control is enabled then virtualized packages will be checked out using the current settings for the project (so things like -P4Changelist= will work)

diff --git a/Engine/Source/Editor/VirtualizationEditor/Private/VirtualizeProjectCommandlet.cpp b/Engine/Source/Editor/VirtualizationEditor/Private/VirtualizeProjectCommandlet.cpp
new file mode 100644
--- /dev/null
+++ b/Engine/Source/Editor/VirtualizationEditor/Private/VirtualizeProjectCommandlet.cpp
@@ -0,0 +1,39 @@
+// Copyright Epic Games, Inc. All Rights Reserved.
+
+#include "VirtualizeProjectCommandlet.h"
+
+#include "CommandletUtils.h"
+#include "Virtualization/VirtualizationSystem.h"
+
+UVirtualizeProjectCommandlet::UVirtualizeProjectCommandlet(const FObjectInitializer& ObjectInitializer)
+	: Super(ObjectInitializer)
+{
+
+}
+
+int32 UVirtualizeProjectCommandlet::Main(const FString& Params)
+{
+	TRACE_CPUPROFILER_EVENT_SCOPE(UVirtualizeProjectCommandlet);
+
+	UE_LOG(LogVirtualization, Display, TEXT("Scanning for project packages..."));
+	TArray<FString> Packages = UE::Virtualization::FindAllPackages();
+
+	UE_LOG(LogVirtualization, Display, TEXT("Virtualizing packages..."));
+
+	UE::Virtualization::IVirtualizationSystem& System = UE::Virtualization::IVirtualizationSystem::Get();
+	UE::Virtualization::FVirtualizationResult Result = System.TryVirtualizePackages(Packages, UE::Virtualization::EVirtualizationOptions::Checkout);
+	if (Result.WasSuccessful())
+	{
+		UE_LOG(LogVirtualization, Display, TEXT("Virtualized %d package(s)"), Result.VirtualizedPackages.Num());
+		if (!Result.CheckedOutPackages.IsEmpty())
+		{
+			UE_LOG(LogVirtualization, Display, TEXT("Checked out %d package(s) from revision control"), Result.CheckedOutPackages.Num());
+		}
+
+		return 0;
+	}
+	else
+	{
+		return -1;
+	}
+}
diff --git a/Engine/Source/Editor/VirtualizationEditor/Private/VirtualizeProjectCommandlet.h b/Engine/Source/Editor/VirtualizationEditor/Private/VirtualizeProjectCommandlet.h
new file mode 100644
--- /dev/null
+++ b/Engine/Source/Editor/VirtualizationEditor/Private/VirtualizeProjectCommandlet.h
@@ -0,0 +1,29 @@
+// Copyright Epic Games, Inc. All Rights Reserved.
+
+#pragma once
+
+#include "Commandlets/Commandlet.h"
+#include "UObject/ObjectMacros.h"
+
+#include "VirtualizeProjectCommandlet.generated.h"
+
+/**
+ * Finds all packages in the project and attempts to virtualize them. If revision control
+ * is enabled then the commandlet will attempt to checkout the packages that need modification.
+ *
+ * Because the commmandlet is the VirtualizationEditor module it needs to be invoked
+ * with the command line:
+ * -run=VirtualizationEditor.VirtualizeProject
+ */
+UCLASS()
+class UVirtualizeProjectCommandlet
+	: public UCommandlet
+{
+	GENERATED_UCLASS_BODY()
+
+	//~ Begin UCommandlet Interface
+	virtual int32 Main(const FString& Params) override;
+	//~ End UCommandlet Interface
+
+	static int32 StaticMain(const FString& Params);
+};
