Add a commandlet -run=VirtualizationEditor.ValidateVirtualizedContent to validate that the engine/project/directory does not contain virtualized data

#rnx

- There are some circumstances where we want to be sure that locations do not contain virtualized data. As an example we were to accidently virtualize an engine package and then distribute that with an engine release it will not work for people. This commandlet is designed to be run as part of a CIS to detect this sort of thing and raise errors so that they can be corrected.
- We support three forms of validation
-- '-CheckEngine' which errors if the engine or its plugins contain virtualized data
-- '-CheckProject' which errors if the current project or its plugins contain virtualized data
-- '-CheckDir' which errors if the given directory or its subdirectories contain virtualized data. More than one path can be supplied by using '+' as the delimiter.

diff --git a/Engine/Source/Editor/VirtualizationEditor/Private/CommandletUtils.cpp b/Engine/Source/Editor/VirtualizationEditor/Private/CommandletUtils.cpp
--- a/Engine/Source/Editor/VirtualizationEditor/Private/CommandletUtils.cpp
+++ b/Engine/Source/Editor/VirtualizationEditor/Private/CommandletUtils.cpp
@@ -31,6 +31,22 @@ TArray<FString> FindAllPackages()
 	return PackageNames;
 }
 
+TArray<FString> FindPackagesInDirectory(const FString& DirectoryToSearch)
+{
+	TArray<FString> FilesInPackageFolder;
+	FPackageName::FindPackagesInDirectory(FilesInPackageFolder, DirectoryToSearch);
+
+	TArray<FString> PackageNames;
+	PackageNames.Reserve(FilesInPackageFolder.Num());
+
+	for (const FString& BasePath : FilesInPackageFolder)
+	{
+		PackageNames.Add(FPaths::CreateStandardFilename(BasePath));
+	}
+
+	return PackageNames;
+}
+
 TArray<FString> DiscoverPackages(const FString& CmdlineParams)
 {
 	TRACE_CPUPROFILER_EVENT_SCOPE(DiscoverPackages);
@@ -38,16 +54,7 @@ TArray<FString> DiscoverPackages(const FString& CmdlineParams)
 	FString PackageDir;
 	if (FParse::Value(*CmdlineParams, TEXT("PackageDir="), PackageDir) || FParse::Value(*CmdlineParams, TEXT("PackageFolder="), PackageDir))
 	{
-		TArray<FString> FilesInPackageFolder;
-		FPackageName::FindPackagesInDirectory(FilesInPackageFolder, PackageDir);
-
-		TArray<FString> PackageNames;
-		for (const FString& BasePath : FilesInPackageFolder)
-		{
-			PackageNames.Add(FPaths::CreateStandardFilename(BasePath));
-		}
-
-		return PackageNames;
+		return FindPackagesInDirectory(PackageDir);
 	}
 	else
 	{
@@ -55,8 +62,6 @@ TArray<FString> DiscoverPackages(const FString& CmdlineParams)
 	}
 }
 
-
-
 TArray<FIoHash> FindVirtualizedPayloads(const TArray<FString>& PackagePaths)
 {
 	TRACE_CPUPROFILER_EVENT_SCOPE(FindVirtualizedPayloads);
diff --git a/Engine/Source/Editor/VirtualizationEditor/Private/CommandletUtils.h b/Engine/Source/Editor/VirtualizationEditor/Private/CommandletUtils.h
--- a/Engine/Source/Editor/VirtualizationEditor/Private/CommandletUtils.h
+++ b/Engine/Source/Editor/VirtualizationEditor/Private/CommandletUtils.h
@@ -9,12 +9,15 @@
 namespace UE::Virtualization
 {
 
-/** Parse all of the active mount points and find all .uasset/.umaps */
+/** Parse all of the active mount points and find all packages */
 TArray<FString> FindAllPackages();
 
+/** Finds all of the packages under a given directory including its subdirectories */
+TArray<FString> FindPackagesInDirectory(const FString& DirectoryToSearch);
+
 /**
  * Finds all of the packages under a the directory given by the provided command line.
- * If no commandline switch can be found then the function will return all avaliable 
+ * If no commandline switch can be found then the function will return all available 
  * packages.
  * Valid commandline switches:
  * '-PackageDir=...'
diff --git a/Engine/Source/Editor/VirtualizationEditor/Private/ValidateVirtualizedContentCommandlet.cpp b/Engine/Source/Editor/VirtualizationEditor/Private/ValidateVirtualizedContentCommandlet.cpp
new file mode 100644
--- /dev/null
+++ b/Engine/Source/Editor/VirtualizationEditor/Private/ValidateVirtualizedContentCommandlet.cpp
@@ -0,0 +1,208 @@
+// Copyright Epic Games, Inc. All Rights Reserved.
+
+#include "ValidateVirtualizedContentCommandlet.h"
+
+#include "CommandletUtils.h"
+#include "HAL/FileManager.h"
+#include "Misc/PackageName.h"
+#include "Misc/Paths.h"
+#include "UObject/PackageTrailer.h"
+
+UValidateVirtualizedContentCommandlet::UValidateVirtualizedContentCommandlet(const FObjectInitializer& ObjectInitializer)
+	: Super(ObjectInitializer)
+{
+}
+
+int32 UValidateVirtualizedContentCommandlet::Main(const FString& Params)
+{
+	TRACE_CPUPROFILER_EVENT_SCOPE(UGeneratePayloadManifestCommandlet);
+	
+	TArray<FString> Tokens;
+	TArray<FString> Switches;
+
+	ParseCommandLine(*Params, Tokens, Switches);
+
+	bool bNoVAContentInEngine = false;
+	bool bNoVAPContentInProject = false;
+
+	TArray<FString> DirectoriesToCheck;
+
+	for (const FString& Switch : Switches)
+	{
+		FString InputPath;
+
+		if (Switch == TEXT("CheckEngine"))
+		{
+			bNoVAContentInEngine = true;
+		}
+		else if (Switch == TEXT("CheckProject"))
+		{
+			bNoVAPContentInProject = true;
+		}
+		else if (FParse::Value(*Switch, TEXT("CheckDir="), InputPath))
+		{
+			InputPath.ParseIntoArray(DirectoriesToCheck, TEXT("+"), true);
+		}
+	}
+
+	TArray<FString> EnginePackages;
+	TArray<FString> ProjectPackages;
+
+	if(bNoVAContentInEngine || bNoVAPContentInProject)
+	{
+		TRACE_CPUPROFILER_EVENT_SCOPE(SortAllPackages);
+
+		TArray<FString> AllPackages = UE::Virtualization::FindAllPackages();
+
+		const FString EngineDir = FPaths::EngineDir();
+
+		for (const FString& Path : AllPackages)
+		{
+			if (bNoVAContentInEngine && Path.StartsWith(EngineDir))
+			{
+				EnginePackages.Add(Path);
+			}
+			else if (bNoVAPContentInProject)
+			{
+				ProjectPackages.Add(Path);
+			}
+		}
+	}
+
+	bool bAllContentValid = true;
+
+	if (bNoVAContentInEngine)
+	{
+		if (!TryValidateContent(TEXT("Engine"), EnginePackages))
+		{
+			bAllContentValid = false;
+		}
+	}
+
+	if (bNoVAPContentInProject)
+	{
+		if (!TryValidateContent(TEXT("Project"), ProjectPackages))
+		{
+			bAllContentValid = false;
+		}
+	}
+
+	if (!DirectoriesToCheck.IsEmpty())
+	{
+		for (const FString& Directory : DirectoriesToCheck)
+		{
+			if (!TryValidateDirectory(Directory))
+			{
+				bAllContentValid = false;
+			}
+		}
+	}
+
+	UE_LOG(LogVirtualization, Display, TEXT("********************************************************************************"));
+
+	return bAllContentValid ? 0 : 1;
+}
+
+TArray<FString> UValidateVirtualizedContentCommandlet::FindVirtualizedPackages(const TArray<FString>& PackagePaths)
+{
+	TRACE_CPUPROFILER_EVENT_SCOPE(ParsePackageTrailers);
+
+	TArray<FString> VirtualizedPackages;
+		
+	for (const FString& Path : PackagePaths)
+	{
+		UE::FPackageTrailer Trailer;
+		if (UE::FPackageTrailer::TryLoadFromFile(Path, Trailer))
+		{
+			const int32 NumVirtualizedPayloads = Trailer.GetNumPayloads(UE::EPayloadStorageType::Virtualized);
+			if (NumVirtualizedPayloads > 0)
+			{
+				FString PackageName;
+				if (FPackageName::TryConvertFilenameToLongPackageName(Path, PackageName))
+				{
+					VirtualizedPackages.Emplace(MoveTemp(PackageName));
+				}
+				else
+				{
+					VirtualizedPackages.Add(Path);
+				}
+				
+			}
+		}
+	}
+
+	return VirtualizedPackages;
+}
+
+bool UValidateVirtualizedContentCommandlet::TryValidateContent(const TCHAR* DebugName, const TArray<FString>& Packages)
+{
+	TRACE_CPUPROFILER_EVENT_SCOPE_TEXT(*WriteToString<128>("TryValidateContent - ", DebugName));
+
+	check(DebugName != nullptr);
+
+	UE_LOG(LogVirtualization, Display, TEXT("********************************************************************************"));
+	UE_LOG(LogVirtualization, Display, TEXT("Looking for virtualized payloads in %s content..."), DebugName);
+	UE_LOG(LogVirtualization, Display, TEXT("Found %d %s package(s)"), Packages.Num(), DebugName);
+
+	TArray<FString> VirtualizedPackages = FindVirtualizedPackages(Packages);
+
+	if (VirtualizedPackages.IsEmpty())
+	{
+		UE_LOG(LogVirtualization, Display, TEXT("No virtualized packages were found in %s content"), DebugName);
+		return true;
+	}
+	else
+	{
+		for (FString& Path : VirtualizedPackages)
+		{
+			UE_LOG(LogVirtualization, Error, TEXT("Package '%s' contains virtualized payloads"), *Path);
+		}
+
+		UE_LOG(LogVirtualization, Error, TEXT("Found %d virtualized package(s) in %s content"), VirtualizedPackages.Num(), DebugName);
+		return false;
+	}
+}
+
+bool UValidateVirtualizedContentCommandlet::TryValidateDirectory(const FString& Directory)
+{
+	TRACE_CPUPROFILER_EVENT_SCOPE(TryValidateDirectory);
+
+	UE_LOG(LogVirtualization, Display, TEXT("********************************************************************************"));
+	UE_LOG(LogVirtualization, Display, TEXT("Searching directory '%s' for virtualized packages..."), *Directory);
+
+	if (!IFileManager::Get().DirectoryExists(*Directory))
+	{
+		UE_LOG(LogVirtualization, Error, TEXT("Directory '%s' could not be found!"), *Directory);
+		return false;
+	}
+
+	TArray<FString> DirectoryPackages;
+	DirectoryPackages = UE::Virtualization::FindPackagesInDirectory(Directory);
+
+	if (DirectoryPackages.IsEmpty())
+	{
+		UE_LOG(LogVirtualization, Display, TEXT("Found no packages under '%s'"), *Directory);
+		return true;
+	}
+
+	UE_LOG(LogVirtualization, Display, TEXT("Found %d package(s) under '%s'"), DirectoryPackages.Num(), *Directory);
+	UE_LOG(LogVirtualization, Display, TEXT("Looking for virtualized payloads under directory..."));
+
+	TArray<FString> VirtualizedPackages = FindVirtualizedPackages(DirectoryPackages);
+
+	if (VirtualizedPackages.IsEmpty())
+	{
+		UE_LOG(LogVirtualization, Display, TEXT("No virtualized packages were found under '%s'"), *Directory);
+	}
+	else
+	{
+		for (FString& Path : VirtualizedPackages)
+		{
+			UE_LOG(LogVirtualization, Error, TEXT("Package '%s' contains virtualized payloads"), *Path);
+		}
+
+		UE_LOG(LogVirtualization, Error, TEXT("Found %d virtualized package(s) under '%s'"), VirtualizedPackages.Num(), *Directory);
+	}
+
+	return true;
+}
diff --git a/Engine/Source/Editor/VirtualizationEditor/Private/ValidateVirtualizedContentCommandlet.h b/Engine/Source/Editor/VirtualizationEditor/Private/ValidateVirtualizedContentCommandlet.h
new file mode 100644
--- /dev/null
+++ b/Engine/Source/Editor/VirtualizationEditor/Private/ValidateVirtualizedContentCommandlet.h
@@ -0,0 +1,51 @@
+// Copyright Epic Games, Inc. All Rights Reserved.
+
+#pragma once
+
+#include "Commandlets/Commandlet.h"
+#include "UObject/ObjectMacros.h"
+
+#include "ValidateVirtualizedContentCommandlet.generated.h"
+
+/**
+ * Used to validate that content does not contain virtualized payloads. The current 
+ * cmdline options are:
+ * -CheckEngine
+ *     Checks all packages currently mounted in the engine
+ * -CheckProject
+ *     Checks all packages currently mounted in the current project
+ * -CheckDir=XYZ (XYZ is the path to the directory, use '+' as the delimiter if supplying
+ * more than one path)
+ *     Checks packages in the given directory and all subdirectories
+ * any or all of which can be passed to the commandlet.
+ * 
+ * If virtualized payloads are found then the package path (or file path if CheckDir 
+ * indicates a directory that is not currently mounted by the porject) is logged 
+ * as an error and the commandlet will eventually return 1 as a failure value.
+ *
+ * Because the commmandlet is the VirtualizationEditor module it needs to be invoked
+ * with the command line:
+ * -run=VirtualizationEditor.ValidateVirtualizedContent
+ * 
+ * Followed by
+ */
+UCLASS()
+class UValidateVirtualizedContentCommandlet
+	: public UCommandlet
+{
+	GENERATED_UCLASS_BODY()
+
+	//~ Begin UCommandlet Interface
+	virtual int32 Main(const FString& Params) override;
+	//~ End UCommandlet Interface
+
+	static int32 StaticMain(const FString& Params);
+
+private:
+
+	TArray<FString> FindVirtualizedPackages(const TArray<FString>& PackagePaths);
+
+	bool TryValidateContent(const TCHAR* DebugName, const TArray<FString>& Packages);
+	bool TryValidateDirectory(const FString& Directory);
+
+};
