Virtualization commandlets now exclude engine packages.

- In CL 24208899 we removed the option to virtualize engine content which means that the virtualization commandlets can use NORMALIZE_ExcludeEnginePackages when searching for packages.
- This helps make our logging a bit more readable, for example if someone were to run the virtualize project commandlet on a new project containing only a single package we will no longer return the number of engine packages + 1 when reporting how many packages are being considered for virtualization. This should help the user to have more confidence in the numbers that we are reporting.

diff --git a/Engine/Source/Editor/VirtualizationEditor/Private/CommandletUtils.cpp b/Engine/Source/Editor/VirtualizationEditor/Private/CommandletUtils.cpp
--- a/Engine/Source/Editor/VirtualizationEditor/Private/CommandletUtils.cpp
+++ b/Engine/Source/Editor/VirtualizationEditor/Private/CommandletUtils.cpp
@@ -15,10 +15,7 @@ TArray<FString> FindAllPackages()
 {
 	TRACE_CPUPROFILER_EVENT_SCOPE(FindAllPackages);
 
-	// TODO: In theory we could check the VA settings and skip engine content if they are 
-	// excluded from virtualization. Not convinced it saves enough time to be worth the 
-	// effort.
-	const EPackageNormalizationFlags PackageFilter = NORMALIZE_DefaultFlags;
+	const uint8 PackageFilter = NORMALIZE_DefaultFlags | NORMALIZE_ExcludeEnginePackages;
 
 	const FString AssetSearch = TEXT("*") + FPackageName::GetAssetPackageExtension();
 	const FString MapSearch = TEXT("*") + FPackageName::GetMapPackageExtension();
