Speed up package discovery for VA commandlets.

#rnx

- Replace use of ::NormalizePackageNames with the asset registry for finding packages.
- We can use the async version of SearchAllAssets so that we don't have to wait on the asset registry cache being saved to disk.
- ::FindAllPackages now returns all paths in our standard form (relative to Engine\Binaries\<Platform>) where as NormalizePackageNames would return a mix of formats depending on the mount point location.

diff --git a/Engine/Source/Editor/VirtualizationEditor/Private/CommandletUtils.cpp b/Engine/Source/Editor/VirtualizationEditor/Private/CommandletUtils.cpp
--- a/Engine/Source/Editor/VirtualizationEditor/Private/CommandletUtils.cpp
+++ b/Engine/Source/Editor/VirtualizationEditor/Private/CommandletUtils.cpp
@@ -2,10 +2,12 @@
 
 #include "CommandletUtils.h"
 
+#include "AssetRegistry/AssetData.h"
+#include "AssetRegistry/AssetRegistryModule.h"
 #include "Async/ParallelFor.h"
 #include "Containers/Set.h"
 #include "Misc/PackageName.h"
-#include "PackageHelperFunctions.h"
+#include "Misc/Paths.h"
 #include "UObject/PackageTrailer.h"
 
 namespace UE::Virtualization
@@ -15,17 +17,39 @@ TArray<FString> FindAllPackages()
 {
 	TRACE_CPUPROFILER_EVENT_SCOPE(FindAllPackages);
 
-	const uint8 PackageFilter = NORMALIZE_DefaultFlags | NORMALIZE_ExcludeEnginePackages;
+	TArray<FString> PackagePaths;
 
-	const FString AssetSearch = TEXT("*") + FPackageName::GetAssetPackageExtension();
-	const FString MapSearch = TEXT("*") + FPackageName::GetMapPackageExtension();
+	FAssetRegistryModule& AssetRegistry = FModuleManager::LoadModuleChecked<FAssetRegistryModule>(TEXT("AssetRegistry"));
 
-	TArray<FString> PackageNames;
-	TArray<FString> Unused;
-	bool bAnyFound = NormalizePackageNames(Unused, PackageNames, AssetSearch, PackageFilter);
-	bAnyFound |= NormalizePackageNames(Unused, PackageNames, MapSearch, PackageFilter);
+	// Do an async search even though we immediately block on it. This will result in the asset registry cache
+	// being saved to disk on a background thread which is an operation we don't need to wait on. This can 
+	// save a fair amount of time on larger projects.
+	const bool bSynchronousSearch = false;
+	AssetRegistry.Get().SearchAllAssets(bSynchronousSearch);
+	AssetRegistry.Get().WaitForCompletion();
 
-	return PackageNames;
+	const FString EnginePath = FPaths::EngineDir();
+
+	AssetRegistry.Get().EnumerateAllPackages([&PackagePaths, EnginePath](FName PackageName, const FAssetPackageData& PackageData)
+		{
+			FString RelFileName;
+			if (PackageData.Extension != EPackageExtension::Unspecified && PackageData.Extension != EPackageExtension::Custom)
+			{
+				const FString Extension = LexToString(PackageData.Extension);
+				if (FPackageName::TryConvertLongPackageNameToFilename(PackageName.ToString(), RelFileName, Extension))
+				{
+					FString StdFileName = FPaths::CreateStandardFilename(RelFileName);
+				
+					// Now we have the absolute file path we can filter out engine packages
+					if (!StdFileName.StartsWith(EnginePath))
+					{
+						PackagePaths.Emplace(MoveTemp(StdFileName));
+					}
+				}
+			}
+		});
+
+	return PackagePaths;
 }
 
 TArray<FString> FindPackagesInDirectory(const FString& DirectoryToSearch)
