Add '-AutoCheckIn' cmdline option to the Virtualize/Rehydrate project commandlets.

- The option will only checkin packages that were checked out by the commandlet. Local files or packages that were already open for edit will not be submitted.

diff --git a/Engine/Source/Editor/VirtualizationEditor/Private/RehydrateProjectCommandlet.cpp b/Engine/Source/Editor/VirtualizationEditor/Private/RehydrateProjectCommandlet.cpp
--- a/Engine/Source/Editor/VirtualizationEditor/Private/RehydrateProjectCommandlet.cpp
+++ b/Engine/Source/Editor/VirtualizationEditor/Private/RehydrateProjectCommandlet.cpp
@@ -5,6 +5,11 @@
 #include "CommandletUtils.h"
 #include "Virtualization/VirtualizationSystem.h"
 
+#include "ISourceControlModule.h"
+#include "ISourceControlProvider.h"
+#include "SourceControlHelpers.h"
+#include "SourceControlOperations.h"
+
 URehydrateProjectCommandlet::URehydrateProjectCommandlet(const FObjectInitializer& ObjectInitializer)
 	: Super(ObjectInitializer)
 {
@@ -15,9 +20,15 @@ int32 URehydrateProjectCommandlet::Main(const FString& Params)
 {
 	TRACE_CPUPROFILER_EVENT_SCOPE(URehydrateProjectCommandlet);
 
+	if (!ParseCmdline(Params))
+	{
+		UE_LOG(LogVirtualization, Error, TEXT("Failed to parse the command line correctly"));
+		return -1;
+	}
+
 	UE_LOG(LogVirtualization, Display, TEXT("Scanning for project packages..."));
 	TArray<FString> Packages = UE::Virtualization::FindAllPackages();
-	
+
 	UE_LOG(LogVirtualization, Display, TEXT("Rehydrating packages..."));
 
 	UE::Virtualization::IVirtualizationSystem& System = UE::Virtualization::IVirtualizationSystem::Get();
@@ -30,6 +41,30 @@ int32 URehydrateProjectCommandlet::Main(const FString& Params)
 			UE_LOG(LogVirtualization, Display, TEXT("Checked out %d package(s) from revision control"), Result.CheckedOutPackages.Num());
 		}
 
+		if (bAutoCheckIn && !Result.CheckedOutPackages.IsEmpty())
+		{
+			UE_LOG(LogVirtualization, Display, TEXT("Attempting to check in modified packages to revision control..."));
+
+			TArray<FString> CheckedOutPackages = SourceControlHelpers::PackageFilenames(Result.CheckedOutPackages);
+
+			const FText FinalDescription = FText::FromString(TEXT("Rehydrating project packages - automated submit from the RehydrateProject commandlet"));
+
+			TSharedRef<FCheckIn, ESPMode::ThreadSafe> CheckInOperation = ISourceControlOperation::Create<FCheckIn>();
+			CheckInOperation->SetDescription(FinalDescription);
+
+			ISourceControlProvider& SourceControlProvider = ISourceControlModule::Get().GetProvider();
+			if (SourceControlProvider.Execute(CheckInOperation, CheckedOutPackages) == ECommandResult::Succeeded)
+			{
+				UE_LOG(LogVirtualization, Display, TEXT("Success! - %s"), *CheckInOperation->GetSuccessMessage().ToString());
+			}
+			else
+			{
+				// The FCheckIn operation tends to log error messages so we don't need to extract the result info
+				// from the operation.
+				UE_LOG(LogVirtualization, Error, TEXT("Failed to check in the package(s) see the above LogSourceControl errors"));
+			}
+		}
+
 		return 0;
 	}
 	else
@@ -37,3 +72,15 @@ int32 URehydrateProjectCommandlet::Main(const FString& Params)
 		return -1;
 	}
 }
+
+bool URehydrateProjectCommandlet::ParseCmdline(const FString& Params)
+{
+	TArray<FString> Tokens;
+	TArray<FString> Switches;
+
+	ParseCommandLine(*Params, Tokens, Switches);
+
+	bAutoCheckIn = Switches.Contains(TEXT("AutoCheckIn"));
+
+	return true;
+}
diff --git a/Engine/Source/Editor/VirtualizationEditor/Private/RehydrateProjectCommandlet.h b/Engine/Source/Editor/VirtualizationEditor/Private/RehydrateProjectCommandlet.h
--- a/Engine/Source/Editor/VirtualizationEditor/Private/RehydrateProjectCommandlet.h
+++ b/Engine/Source/Editor/VirtualizationEditor/Private/RehydrateProjectCommandlet.h
@@ -19,6 +19,8 @@ UCLASS()
 class URehydrateProjectCommandlet
 	: public UCommandlet
 {
+private:
+
 	GENERATED_UCLASS_BODY()
 
 	//~ Begin UCommandlet Interface
@@ -26,4 +28,10 @@ class URehydrateProjectCommandlet
 	//~ End UCommandlet Interface
 
 	static int32 StaticMain(const FString& Params);
+
+	bool ParseCmdline(const FString& Params);
+
+private:
+
+	bool bAutoCheckIn = false;
 };
diff --git a/Engine/Source/Editor/VirtualizationEditor/Private/VirtualizeProjectCommandlet.cpp b/Engine/Source/Editor/VirtualizationEditor/Private/VirtualizeProjectCommandlet.cpp
--- a/Engine/Source/Editor/VirtualizationEditor/Private/VirtualizeProjectCommandlet.cpp
+++ b/Engine/Source/Editor/VirtualizationEditor/Private/VirtualizeProjectCommandlet.cpp
@@ -5,6 +5,11 @@
 #include "CommandletUtils.h"
 #include "Virtualization/VirtualizationSystem.h"
 
+#include "ISourceControlModule.h"
+#include "ISourceControlProvider.h"
+#include "SourceControlHelpers.h"
+#include "SourceControlOperations.h"
+
 UVirtualizeProjectCommandlet::UVirtualizeProjectCommandlet(const FObjectInitializer& ObjectInitializer)
 	: Super(ObjectInitializer)
 {
@@ -15,6 +20,12 @@ int32 UVirtualizeProjectCommandlet::Main(const FString& Params)
 {
 	TRACE_CPUPROFILER_EVENT_SCOPE(UVirtualizeProjectCommandlet);
 
+	if (!ParseCmdline(Params))
+	{
+		UE_LOG(LogVirtualization, Error, TEXT("Failed to parse the command line correctly"));
+		return -1;
+	}
+
 	UE_LOG(LogVirtualization, Display, TEXT("Scanning for project packages..."));
 	TArray<FString> Packages = UE::Virtualization::FindAllPackages();
 
@@ -30,6 +41,38 @@ int32 UVirtualizeProjectCommandlet::Main(const FString& Params)
 			UE_LOG(LogVirtualization, Display, TEXT("Checked out %d package(s) from revision control"), Result.CheckedOutPackages.Num());
 		}
 
+		if (bAutoCheckIn && !Result.CheckedOutPackages.IsEmpty())
+		{
+			UE_LOG(LogVirtualization, Display, TEXT("Attempting to check in modified packages to revision control..."));
+
+			TArray<FString> CheckedOutPackages = SourceControlHelpers::PackageFilenames(Result.CheckedOutPackages);
+
+			FTextBuilder DescriptionBuilder;
+			DescriptionBuilder.AppendLine(FText::FromString(TEXT("Virtualizing project packages - automated submit from the VirtualizeProject commandlet")));
+
+			for (const FText& Line : Result.DescriptionTags)
+			{
+				DescriptionBuilder.AppendLine(Line);
+			}
+
+			const FText CheckInDescription = DescriptionBuilder.ToText();
+
+			TSharedRef<FCheckIn, ESPMode::ThreadSafe> CheckInOperation = ISourceControlOperation::Create<FCheckIn>();
+			CheckInOperation->SetDescription(CheckInDescription);
+
+			ISourceControlProvider& SourceControlProvider = ISourceControlModule::Get().GetProvider();
+			if (SourceControlProvider.Execute(CheckInOperation, CheckedOutPackages) == ECommandResult::Succeeded)
+			{
+				UE_LOG(LogVirtualization, Display, TEXT("Success! - %s"), *CheckInOperation->GetSuccessMessage().ToString());
+			}
+			else
+			{
+				// The FCheckIn operation tends to log error messages so we don't need to extract the result info
+				// from the operation.
+				UE_LOG(LogVirtualization, Error, TEXT("Failed to check in the package(s) see the above LogSourceControl errors"));
+			}
+		}
+
 		return 0;
 	}
 	else
@@ -37,3 +80,15 @@ int32 UVirtualizeProjectCommandlet::Main(const FString& Params)
 		return -1;
 	}
 }
+
+bool UVirtualizeProjectCommandlet::ParseCmdline(const FString& Params)
+{
+	TArray<FString> Tokens;
+	TArray<FString> Switches;
+
+	ParseCommandLine(*Params, Tokens, Switches);
+
+	bAutoCheckIn = Switches.Contains(TEXT("AutoCheckIn"));
+
+	return true;
+}
diff --git a/Engine/Source/Editor/VirtualizationEditor/Private/VirtualizeProjectCommandlet.h b/Engine/Source/Editor/VirtualizationEditor/Private/VirtualizeProjectCommandlet.h
--- a/Engine/Source/Editor/VirtualizationEditor/Private/VirtualizeProjectCommandlet.h
+++ b/Engine/Source/Editor/VirtualizationEditor/Private/VirtualizeProjectCommandlet.h
@@ -19,6 +19,8 @@ UCLASS()
 class UVirtualizeProjectCommandlet
 	: public UCommandlet
 {
+private:
+
 	GENERATED_UCLASS_BODY()
 
 	//~ Begin UCommandlet Interface
@@ -26,4 +28,10 @@ class UVirtualizeProjectCommandlet
 	//~ End UCommandlet Interface
 
 	static int32 StaticMain(const FString& Params);
+
+	bool ParseCmdline(const FString& Params);
+
+private:
+
+	bool bAutoCheckIn = false;
 };
diff --git a/Engine/Source/Editor/VirtualizationEditor/VirtualizationEditor.Build.cs b/Engine/Source/Editor/VirtualizationEditor/VirtualizationEditor.Build.cs
--- a/Engine/Source/Editor/VirtualizationEditor/VirtualizationEditor.Build.cs
+++ b/Engine/Source/Editor/VirtualizationEditor/VirtualizationEditor.Build.cs
@@ -16,6 +16,7 @@ public class VirtualizationEditor : ModuleRules
 				"Engine",
 				"UnrealEd",
 				"ToolMenus",
+				"SourceControl",
 				"Slate",
 				"SlateCore",
 				"Virtualization"
