The VA commandlet 'ValidateVirtualizedContent' now iterates over all packages in a project and makes sure that any virtualized payloads encountered can be found in persistent storage. The original version of this commandlet has been renamed to 'CheckForVirtualizedContent' and continues to be used to make sure that virtualized packages do not get submitted to places that they should not.

- Renamed UValidateVirtualizedContentCommandlet to UCheckForVirtualizedContentCommandlet. P4 rename was not done because I wanted to repurpose UValidateVirtualizedContentCommandlet and adding new versions of the code files after renaming them is likely to cause weird p4 merge issues in the future.

diff --git a/Engine/Source/Editor/VirtualizationEditor/Private/CheckForVirtualizedContentCommandlet.cpp b/Engine/Source/Editor/VirtualizationEditor/Private/CheckForVirtualizedContentCommandlet.cpp
new file mode 100644
--- /dev/null
+++ b/Engine/Source/Editor/VirtualizationEditor/Private/CheckForVirtualizedContentCommandlet.cpp
@@ -0,0 +1,208 @@
+// Copyright Epic Games, Inc. All Rights Reserved.
+
+#include "CheckForVirtualizedContentCommandlet.h"
+
+#include "CommandletUtils.h"
+#include "HAL/FileManager.h"
+#include "Misc/PackageName.h"
+#include "Misc/Paths.h"
+#include "UObject/PackageTrailer.h"
+
+UCheckForVirtualizedContentCommandlet::UCheckForVirtualizedContentCommandlet(const FObjectInitializer& ObjectInitializer)
+	: Super(ObjectInitializer)
+{
+}
+
+int32 UCheckForVirtualizedContentCommandlet::Main(const FString& Params)
+{
+	TRACE_CPUPROFILER_EVENT_SCOPE(UCheckForVirtualizedContentCommandlet);
+
+	TArray<FString> Tokens;
+	TArray<FString> Switches;
+
+	ParseCommandLine(*Params, Tokens, Switches);
+
+	bool bNoVAContentInEngine = false;
+	bool bNoVAPContentInProject = false;
+
+	TArray<FString> DirectoriesToCheck;
+
+	for (const FString& Switch : Switches)
+	{
+		FString InputPath;
+
+		if (Switch == TEXT("CheckEngine"))
+		{
+			bNoVAContentInEngine = true;
+		}
+		else if (Switch == TEXT("CheckProject"))
+		{
+			bNoVAPContentInProject = true;
+		}
+		else if (FParse::Value(*Switch, TEXT("CheckDir="), InputPath))
+		{
+			InputPath.ParseIntoArray(DirectoriesToCheck, TEXT("+"), true);
+		}
+	}
+
+	TArray<FString> EnginePackages;
+	TArray<FString> ProjectPackages;
+
+	if (bNoVAContentInEngine || bNoVAPContentInProject)
+	{
+		TRACE_CPUPROFILER_EVENT_SCOPE(SortAllPackages);
+
+		TArray<FString> AllPackages = UE::Virtualization::FindAllPackages();
+
+		const FString EngineDir = FPaths::EngineDir();
+
+		for (const FString& Path : AllPackages)
+		{
+			if (bNoVAContentInEngine && Path.StartsWith(EngineDir))
+			{
+				EnginePackages.Add(Path);
+			}
+			else if (bNoVAPContentInProject)
+			{
+				ProjectPackages.Add(Path);
+			}
+		}
+	}
+
+	bool bAllContentValid = true;
+
+	if (bNoVAContentInEngine)
+	{
+		if (!TryValidateContent(TEXT("Engine"), EnginePackages))
+		{
+			bAllContentValid = false;
+		}
+	}
+
+	if (bNoVAPContentInProject)
+	{
+		if (!TryValidateContent(TEXT("Project"), ProjectPackages))
+		{
+			bAllContentValid = false;
+		}
+	}
+
+	if (!DirectoriesToCheck.IsEmpty())
+	{
+		for (const FString& Directory : DirectoriesToCheck)
+		{
+			if (!TryValidateDirectory(Directory))
+			{
+				bAllContentValid = false;
+			}
+		}
+	}
+
+	UE_LOG(LogVirtualization, Display, TEXT("********************************************************************************"));
+
+	return bAllContentValid ? 0 : 1;
+}
+
+TArray<FString> UCheckForVirtualizedContentCommandlet::FindVirtualizedPackages(const TArray<FString>& PackagePaths)
+{
+	TRACE_CPUPROFILER_EVENT_SCOPE(ParsePackageTrailers);
+
+	TArray<FString> VirtualizedPackages;
+
+	for (const FString& Path : PackagePaths)
+	{
+		UE::FPackageTrailer Trailer;
+		if (UE::FPackageTrailer::TryLoadFromFile(Path, Trailer))
+		{
+			const int32 NumVirtualizedPayloads = Trailer.GetNumPayloads(UE::EPayloadStorageType::Virtualized);
+			if (NumVirtualizedPayloads > 0)
+			{
+				FString PackageName;
+				if (FPackageName::TryConvertFilenameToLongPackageName(Path, PackageName))
+				{
+					VirtualizedPackages.Emplace(MoveTemp(PackageName));
+				}
+				else
+				{
+					VirtualizedPackages.Add(Path);
+				}
+
+			}
+		}
+	}
+
+	return VirtualizedPackages;
+}
+
+bool UCheckForVirtualizedContentCommandlet::TryValidateContent(const TCHAR* DebugName, const TArray<FString>& Packages)
+{
+	TRACE_CPUPROFILER_EVENT_SCOPE_TEXT(*WriteToString<128>("TryValidateContent - ", DebugName));
+
+	check(DebugName != nullptr);
+
+	UE_LOG(LogVirtualization, Display, TEXT("********************************************************************************"));
+	UE_LOG(LogVirtualization, Display, TEXT("Looking for virtualized payloads in %s content..."), DebugName);
+	UE_LOG(LogVirtualization, Display, TEXT("Found %d %s package(s)"), Packages.Num(), DebugName);
+
+	TArray<FString> VirtualizedPackages = FindVirtualizedPackages(Packages);
+
+	if (VirtualizedPackages.IsEmpty())
+	{
+		UE_LOG(LogVirtualization, Display, TEXT("No virtualized packages were found in %s content"), DebugName);
+		return true;
+	}
+	else
+	{
+		for (FString& Path : VirtualizedPackages)
+		{
+			UE_LOG(LogVirtualization, Error, TEXT("Package '%s' contains virtualized payloads"), *Path);
+		}
+
+		UE_LOG(LogVirtualization, Error, TEXT("Found %d virtualized package(s) in %s content"), VirtualizedPackages.Num(), DebugName);
+		return false;
+	}
+}
+
+bool UCheckForVirtualizedContentCommandlet::TryValidateDirectory(const FString& Directory)
+{
+	TRACE_CPUPROFILER_EVENT_SCOPE(TryValidateDirectory);
+
+	UE_LOG(LogVirtualization, Display, TEXT("********************************************************************************"));
+	UE_LOG(LogVirtualization, Display, TEXT("Searching directory '%s' for virtualized packages..."), *Directory);
+
+	if (!IFileManager::Get().DirectoryExists(*Directory))
+	{
+		UE_LOG(LogVirtualization, Error, TEXT("Directory '%s' could not be found!"), *Directory);
+		return false;
+	}
+
+	TArray<FString> DirectoryPackages;
+	DirectoryPackages = UE::Virtualization::FindPackagesInDirectory(Directory);
+
+	if (DirectoryPackages.IsEmpty())
+	{
+		UE_LOG(LogVirtualization, Display, TEXT("Found no packages under '%s'"), *Directory);
+		return true;
+	}
+
+	UE_LOG(LogVirtualization, Display, TEXT("Found %d package(s) under '%s'"), DirectoryPackages.Num(), *Directory);
+	UE_LOG(LogVirtualization, Display, TEXT("Looking for virtualized payloads under directory..."));
+
+	TArray<FString> VirtualizedPackages = FindVirtualizedPackages(DirectoryPackages);
+
+	if (VirtualizedPackages.IsEmpty())
+	{
+		UE_LOG(LogVirtualization, Display, TEXT("No virtualized packages were found under '%s'"), *Directory);
+	}
+	else
+	{
+		for (FString& Path : VirtualizedPackages)
+		{
+			UE_LOG(LogVirtualization, Error, TEXT("Package '%s' contains virtualized payloads"), *Path);
+		}
+
+		UE_LOG(LogVirtualization, Error, TEXT("Found %d virtualized package(s) under '%s'"), VirtualizedPackages.Num(), *Directory);
+	}
+
+	return true;
+}
diff --git a/Engine/Source/Editor/VirtualizationEditor/Private/CheckForVirtualizedContentCommandlet.h b/Engine/Source/Editor/VirtualizationEditor/Private/CheckForVirtualizedContentCommandlet.h
new file mode 100644
--- /dev/null
+++ b/Engine/Source/Editor/VirtualizationEditor/Private/CheckForVirtualizedContentCommandlet.h
@@ -0,0 +1,51 @@
+// Copyright Epic Games, Inc. All Rights Reserved.
+
+#pragma once
+
+#include "Commandlets/Commandlet.h"
+#include "UObject/ObjectMacros.h"
+
+#include "CheckForVirtualizedContentCommandlet.generated.h"
+
+/**
+ * Used to validate that content does not contain virtualized payloads. The current
+ * cmdline options are:
+ * -CheckEngine
+ *     Checks all packages currently mounted in the engine
+ * -CheckProject
+ *     Checks all packages currently mounted in the current project
+ * -CheckDir=XYZ (XYZ is the path to the directory, use '+' as the delimiter if supplying
+ * more than one path)
+ *     Checks packages in the given directory and all subdirectories
+ * any or all of which can be passed to the commandlet.
+ *
+ * If virtualized payloads are found then the package path (or file path if CheckDir
+ * indicates a directory that is not currently mounted by the project) is logged
+ * as an error and the commandlet will eventually return 1 as a failure value.
+ *
+ * Because the commandlet is the VirtualizationEditor module it needs to be invoked
+ * with the command line:
+ * -run=VirtualizationEditor.CheckForVirtualizedContent
+ *
+ * Followed by
+ */
+UCLASS()
+class UCheckForVirtualizedContentCommandlet
+	: public UCommandlet
+{
+	GENERATED_UCLASS_BODY()
+
+	//~ Begin UCommandlet Interface
+	virtual int32 Main(const FString& Params) override;
+	//~ End UCommandlet Interface
+
+	static int32 StaticMain(const FString& Params);
+
+private:
+
+	TArray<FString> FindVirtualizedPackages(const TArray<FString>& PackagePaths);
+
+	bool TryValidateContent(const TCHAR* DebugName, const TArray<FString>& Packages);
+	bool TryValidateDirectory(const FString& Directory);
+
+};
diff --git a/Engine/Source/Editor/VirtualizationEditor/Private/CommandletUtils.cpp b/Engine/Source/Editor/VirtualizationEditor/Private/CommandletUtils.cpp
--- a/Engine/Source/Editor/VirtualizationEditor/Private/CommandletUtils.cpp
+++ b/Engine/Source/Editor/VirtualizationEditor/Private/CommandletUtils.cpp
@@ -114,4 +114,45 @@ TArray<FIoHash> FindVirtualizedPayloads(const TArray<FString>& PackagePaths)
 	return AllPayloads.Array();
 }
 
+void FindVirtualizedPayloadsAndTrailers(const TArray<FString>& PackagePaths, TMap<FString, UE::FPackageTrailer>& OutPackages, TSet<FIoHash>& OutPayloads)
+{
+	TRACE_CPUPROFILER_EVENT_SCOPE(FindVirtualizedPayloadsAndTrailers);
+
+	struct FTaskContext
+	{
+		TMap<FString, UE::FPackageTrailer> Packages;
+		TSet<FIoHash> Payloads;
+	};
+
+	TArray<FTaskContext> TaskContext;
+
+	ParallelForWithTaskContext(TaskContext, PackagePaths.Num(),
+		[&PackagePaths](FTaskContext& Context, int32 Index)
+		{
+			const FString& PackageName = PackagePaths[Index];
+
+	UE::FPackageTrailer Trailer;
+	if (UE::FPackageTrailer::TryLoadFromFile(PackageName, Trailer))
+	{
+		const TArray<FIoHash> VirtualizedPayloads = Trailer.GetPayloads(UE::EPayloadStorageType::Virtualized);
+		if (!VirtualizedPayloads.IsEmpty())
+		{
+			Context.Packages.Emplace(PackageName, MoveTemp(Trailer));
+			Context.Payloads.Append(VirtualizedPayloads);
+		}
+	}
+		});
+
+	for (const FTaskContext& Context : TaskContext)
+	{
+		OutPackages.Append(Context.Packages);
+		OutPayloads.Append(Context.Payloads);
+	}
+
+	{
+		TRACE_CPUPROFILER_EVENT_SCOPE(CompactStable);
+		OutPayloads.CompactStable();
+	}
+}
+
 } //namespace UE::Virtualization
diff --git a/Engine/Source/Editor/VirtualizationEditor/Private/CommandletUtils.h b/Engine/Source/Editor/VirtualizationEditor/Private/CommandletUtils.h
--- a/Engine/Source/Editor/VirtualizationEditor/Private/CommandletUtils.h
+++ b/Engine/Source/Editor/VirtualizationEditor/Private/CommandletUtils.h
@@ -6,6 +6,8 @@
 #include "Containers/UnrealString.h"
 #include "IO/IoHash.h"
 
+namespace UE { class FPackageTrailer; }
+
 namespace UE::Virtualization
 {
 
@@ -31,4 +33,15 @@ TArray<FString> DiscoverPackages(const FString& CmdlineParams);
 /** Returns a combined list of unique virtualized payload ids from the given list of packages */
 TArray<FIoHash> FindVirtualizedPayloads(const TArray<FString>& PackageNames);
 
+/** 
+ * Load and parse the package trailers for the given packages. We return a map of all of the packages that contain
+ * virtualized payloads and a unique set of all the virtualized payloads referenced. Note that packages can reference
+ * the same payload if they re-use assets.
+ * 
+ * @param PackagePaths	A list of packages to check
+ * @param OutPackages	A map of package paths to package trailers
+ * @param OutPayloads	A unique set of all the virtualized payloads referenced by the packages
+ */
+void FindVirtualizedPayloadsAndTrailers(const TArray<FString>& PackagePaths, TMap<FString, UE::FPackageTrailer>& OutPackages, TSet<FIoHash>& OutPayloads);
+
 } //namespace UE::Virtualization
diff --git a/Engine/Source/Editor/VirtualizationEditor/Private/ValidateVirtualizedContentCommandlet.cpp b/Engine/Source/Editor/VirtualizationEditor/Private/ValidateVirtualizedContentCommandlet.cpp
--- a/Engine/Source/Editor/VirtualizationEditor/Private/ValidateVirtualizedContentCommandlet.cpp
+++ b/Engine/Source/Editor/VirtualizationEditor/Private/ValidateVirtualizedContentCommandlet.cpp
@@ -3,10 +3,8 @@
 #include "ValidateVirtualizedContentCommandlet.h"
 
 #include "CommandletUtils.h"
-#include "HAL/FileManager.h"
-#include "Misc/PackageName.h"
-#include "Misc/Paths.h"
 #include "UObject/PackageTrailer.h"
+#include "Virtualization/VirtualizationSystem.h"
 
 UValidateVirtualizedContentCommandlet::UValidateVirtualizedContentCommandlet(const FObjectInitializer& ObjectInitializer)
 	: Super(ObjectInitializer)
@@ -15,194 +13,75 @@ UValidateVirtualizedContentCommandlet::UValidateVirtualizedContentCommandlet(con
 
 int32 UValidateVirtualizedContentCommandlet::Main(const FString& Params)
 {
-	TRACE_CPUPROFILER_EVENT_SCOPE(UGeneratePayloadManifestCommandlet);
-	
-	TArray<FString> Tokens;
-	TArray<FString> Switches;
+	using namespace UE::Virtualization;
 
-	ParseCommandLine(*Params, Tokens, Switches);
+	TRACE_CPUPROFILER_EVENT_SCOPE(UValidateVirtualizedContentCommandlet);
 
-	bool bNoVAContentInEngine = false;
-	bool bNoVAPContentInProject = false;
+	UE_LOG(LogVirtualization, Display, TEXT("Finding packages in the project..."));
+	TArray<FString> PackagePaths = FindAllPackages();
+	UE_LOG(LogVirtualization, Display, TEXT("Found %d package(s)"), PackagePaths.Num());
 
-	TArray<FString> DirectoriesToCheck;
+	TMap<FString, UE::FPackageTrailer> Packages;
+	TSet<FIoHash> Payloads;
 
-	for (const FString& Switch : Switches)
-	{
-		FString InputPath;
-
-		if (Switch == TEXT("CheckEngine"))
-		{
-			bNoVAContentInEngine = true;
-		}
-		else if (Switch == TEXT("CheckProject"))
-		{
-			bNoVAPContentInProject = true;
-		}
-		else if (FParse::Value(*Switch, TEXT("CheckDir="), InputPath))
-		{
-			InputPath.ParseIntoArray(DirectoriesToCheck, TEXT("+"), true);
-		}
-	}
+	UE_LOG(LogVirtualization, Display, TEXT("Scanning package(s) for virtualized payloads..."), PackagePaths.Num());
+	FindVirtualizedPayloadsAndTrailers(PackagePaths, Packages, Payloads);
+	UE_LOG(LogVirtualization, Display, TEXT("Found %d virtualized package(s) with %d unique payload(s)"), Packages.Num(), Payloads.Num());
 
-	TArray<FString> EnginePackages;
-	TArray<FString> ProjectPackages;
+	IVirtualizationSystem& System = IVirtualizationSystem::Get();
 
-	if(bNoVAContentInEngine || bNoVAPContentInProject)
+	UE_LOG(LogVirtualization, Display, TEXT("Querying the state of the virtualized payload(s) in persistent storage..."));
+	TArray<EPayloadStatus> PayloadStatuses;
+	if (System.QueryPayloadStatuses(Payloads.Array(), EStorageType::Persistent, PayloadStatuses) != EQueryResult::Success)
 	{
-		TRACE_CPUPROFILER_EVENT_SCOPE(SortAllPackages);
-
-		TArray<FString> AllPackages = UE::Virtualization::FindAllPackages();
-
-		const FString EngineDir = FPaths::EngineDir();
-
-		for (const FString& Path : AllPackages)
-		{
-			if (bNoVAContentInEngine && Path.StartsWith(EngineDir))
-			{
-				EnginePackages.Add(Path);
-			}
-			else if (bNoVAPContentInProject)
-			{
-				ProjectPackages.Add(Path);
-			}
-		}
+		UE_LOG(LogVirtualization, Error, TEXT("Failed to query the statuses of the payload(s)"));
+		return 1;
 	}
 
-	bool bAllContentValid = true;
+	int32 ErrorCount = 0;
 
-	if (bNoVAContentInEngine)
 	{
-		if (!TryValidateContent(TEXT("Engine"), EnginePackages))
-		{
-			bAllContentValid = false;
-		}
-	}
+		TRACE_CPUPROFILER_EVENT_SCOPE(ValidatePackages);
 
-	if (bNoVAPContentInProject)
-	{
-		if (!TryValidateContent(TEXT("Project"), ProjectPackages))
-		{
-			bAllContentValid = false;
-		}
-	}
+		UE_LOG(LogVirtualization, Display, TEXT("Checking for missing payloads..."));
 
-	if (!DirectoriesToCheck.IsEmpty())
-	{
-		for (const FString& Directory : DirectoriesToCheck)
+		for (const TPair<FString, UE::FPackageTrailer>& Package : Packages)
 		{
-			if (!TryValidateDirectory(Directory))
-			{
-				bAllContentValid = false;
-			}
-		}
-	}
-
-	UE_LOG(LogVirtualization, Display, TEXT("********************************************************************************"));
+			bool bFoundErrors = false;
 
-	return bAllContentValid ? 0 : 1;
-}
-
-TArray<FString> UValidateVirtualizedContentCommandlet::FindVirtualizedPackages(const TArray<FString>& PackagePaths)
-{
-	TRACE_CPUPROFILER_EVENT_SCOPE(ParsePackageTrailers);
+			const TArray<FIoHash> VirtualizedPayloads = Package.Value.GetPayloads(UE::EPayloadStorageType::Virtualized);
 
-	TArray<FString> VirtualizedPackages;
-		
-	for (const FString& Path : PackagePaths)
-	{
-		UE::FPackageTrailer Trailer;
-		if (UE::FPackageTrailer::TryLoadFromFile(Path, Trailer))
-		{
-			const int32 NumVirtualizedPayloads = Trailer.GetNumPayloads(UE::EPayloadStorageType::Virtualized);
-			if (NumVirtualizedPayloads > 0)
+			for (const FIoHash& PayloadId : VirtualizedPayloads)
 			{
-				FString PackageName;
-				if (FPackageName::TryConvertFilenameToLongPackageName(Path, PackageName))
+				const int32 Index = Payloads.FindId(PayloadId).AsInteger();
+				if (PayloadStatuses[Index] == EPayloadStatus::FoundPartial)
 				{
-					VirtualizedPackages.Emplace(MoveTemp(PackageName));
+					// TODO: We currently don't have a way to inform the user which persistent backend the payload is not in!
+					UE_LOG(LogVirtualization, Error, TEXT("%s: Payload '%s' could not be found in all persistent backends"), *Package.Key, *LexToString(PayloadId));
+					bFoundErrors = true;
 				}
-				else
+				else if (PayloadStatuses[Index] != EPayloadStatus::FoundAll)
 				{
-					VirtualizedPackages.Add(Path);
+					UE_LOG(LogVirtualization, Error, TEXT("%s: Payload '%s' could not be found in any persistent backend"), *Package.Key, *LexToString(PayloadId));
+					bFoundErrors = true;
 				}
-				
 			}
-		}
-	}
-
-	return VirtualizedPackages;
-}
-
-bool UValidateVirtualizedContentCommandlet::TryValidateContent(const TCHAR* DebugName, const TArray<FString>& Packages)
-{
-	TRACE_CPUPROFILER_EVENT_SCOPE_TEXT(*WriteToString<128>("TryValidateContent - ", DebugName));
-
-	check(DebugName != nullptr);
-
-	UE_LOG(LogVirtualization, Display, TEXT("********************************************************************************"));
-	UE_LOG(LogVirtualization, Display, TEXT("Looking for virtualized payloads in %s content..."), DebugName);
-	UE_LOG(LogVirtualization, Display, TEXT("Found %d %s package(s)"), Packages.Num(), DebugName);
 
-	TArray<FString> VirtualizedPackages = FindVirtualizedPackages(Packages);
-
-	if (VirtualizedPackages.IsEmpty())
-	{
-		UE_LOG(LogVirtualization, Display, TEXT("No virtualized packages were found in %s content"), DebugName);
-		return true;
-	}
-	else
-	{
-		for (FString& Path : VirtualizedPackages)
-		{
-			UE_LOG(LogVirtualization, Error, TEXT("Package '%s' contains virtualized payloads"), *Path);
+			if (bFoundErrors)
+			{
+				ErrorCount++;
+			}
 		}
-
-		UE_LOG(LogVirtualization, Error, TEXT("Found %d virtualized package(s) in %s content"), VirtualizedPackages.Num(), DebugName);
-		return false;
 	}
-}
-
-bool UValidateVirtualizedContentCommandlet::TryValidateDirectory(const FString& Directory)
-{
-	TRACE_CPUPROFILER_EVENT_SCOPE(TryValidateDirectory);
 
-	UE_LOG(LogVirtualization, Display, TEXT("********************************************************************************"));
-	UE_LOG(LogVirtualization, Display, TEXT("Searching directory '%s' for virtualized packages..."), *Directory);
-
-	if (!IFileManager::Get().DirectoryExists(*Directory))
+	if (ErrorCount == 0)
 	{
-		UE_LOG(LogVirtualization, Error, TEXT("Directory '%s' could not be found!"), *Directory);
-		return false;
-	}
-
-	TArray<FString> DirectoryPackages;
-	DirectoryPackages = UE::Virtualization::FindPackagesInDirectory(Directory);
-
-	if (DirectoryPackages.IsEmpty())
-	{
-		UE_LOG(LogVirtualization, Display, TEXT("Found no packages under '%s'"), *Directory);
-		return true;
-	}
-
-	UE_LOG(LogVirtualization, Display, TEXT("Found %d package(s) under '%s'"), DirectoryPackages.Num(), *Directory);
-	UE_LOG(LogVirtualization, Display, TEXT("Looking for virtualized payloads under directory..."));
-
-	TArray<FString> VirtualizedPackages = FindVirtualizedPackages(DirectoryPackages);
-
-	if (VirtualizedPackages.IsEmpty())
-	{
-		UE_LOG(LogVirtualization, Display, TEXT("No virtualized packages were found under '%s'"), *Directory);
+		UE_LOG(LogVirtualization, Display, TEXT("All virtualized payloads could be found in persistent storage"));
+		return 0;
 	}
 	else
 	{
-		for (FString& Path : VirtualizedPackages)
-		{
-			UE_LOG(LogVirtualization, Error, TEXT("Package '%s' contains virtualized payloads"), *Path);
-		}
-
-		UE_LOG(LogVirtualization, Error, TEXT("Found %d virtualized package(s) under '%s'"), VirtualizedPackages.Num(), *Directory);
+		UE_LOG(LogVirtualization, Error, TEXT("%d/%d package(s) had at least one virtualized payload missing from persistent storage"), ErrorCount, Packages.Num());
+		return 0;
 	}
-
-	return true;
 }
diff --git a/Engine/Source/Editor/VirtualizationEditor/Private/ValidateVirtualizedContentCommandlet.h b/Engine/Source/Editor/VirtualizationEditor/Private/ValidateVirtualizedContentCommandlet.h
--- a/Engine/Source/Editor/VirtualizationEditor/Private/ValidateVirtualizedContentCommandlet.h
+++ b/Engine/Source/Editor/VirtualizationEditor/Private/ValidateVirtualizedContentCommandlet.h
@@ -8,26 +8,15 @@
 #include "ValidateVirtualizedContentCommandlet.generated.h"
 
 /**
- * Used to validate that content does not contain virtualized payloads. The current 
- * cmdline options are:
- * -CheckEngine
- *     Checks all packages currently mounted in the engine
- * -CheckProject
- *     Checks all packages currently mounted in the current project
- * -CheckDir=XYZ (XYZ is the path to the directory, use '+' as the delimiter if supplying
- * more than one path)
- *     Checks packages in the given directory and all subdirectories
- * any or all of which can be passed to the commandlet.
+ * Iterates over all of the packages in a project and identifies which packages contain
+ * references to virtualized payloads. The commandlet will then check that all virtualized
+ * payloads can be found in persistent storage. Error messages will be logged for
+ * packages that contain virtualized payloads that cannot be found in one or more persistent
+ * storage backend.
  * 
- * If virtualized payloads are found then the package path (or file path if CheckDir 
- * indicates a directory that is not currently mounted by the porject) is logged 
- * as an error and the commandlet will eventually return 1 as a failure value.
- *
- * Because the commmandlet is the VirtualizationEditor module it needs to be invoked
+ * Because the commandlet is the VirtualizationEditor module it needs to be invoked
  * with the command line:
  * -run=VirtualizationEditor.ValidateVirtualizedContent
- * 
- * Followed by
  */
 UCLASS()
 class UValidateVirtualizedContentCommandlet
@@ -40,12 +29,4 @@ class UValidateVirtualizedContentCommandlet
 	//~ End UCommandlet Interface
 
 	static int32 StaticMain(const FString& Params);
-
-private:
-
-	TArray<FString> FindVirtualizedPackages(const TArray<FString>& PackagePaths);
-
-	bool TryValidateContent(const TCHAR* DebugName, const TArray<FString>& Packages);
-	bool TryValidateDirectory(const FString& Directory);
-
 };
