Fix the 'CheckForVirtualizedContent ' commandlet to be able to check engine content once again.

#rnx

- Recently how we discover packages for the VA commandlets was changed to be much faster but it was also changed to automatically filter out all engine content as we removed the optional (opt in) feature of virtualizing engine content.
- This meant that 'CheckForVirtualizedContent' would not be able to check engine packages if '-CheckEngine' was set as we were not providing the commandlet with them.
- ::FindAllPackages has been renamed ::FindPackages and both it and ::DiscoverPackages now accept a flags parameter that allows the caller to request that engine content be excluded.
- Although most commandlets do not want engine content, having the engine content filter as opt in closer matches existing apis and will probably lead to less mistakes.

diff --git a/Engine/Source/Editor/VirtualizationEditor/Private/CheckForVirtualizedContentCommandlet.cpp b/Engine/Source/Editor/VirtualizationEditor/Private/CheckForVirtualizedContentCommandlet.cpp
--- a/Engine/Source/Editor/VirtualizationEditor/Private/CheckForVirtualizedContentCommandlet.cpp
+++ b/Engine/Source/Editor/VirtualizationEditor/Private/CheckForVirtualizedContentCommandlet.cpp
@@ -52,7 +52,7 @@ int32 UCheckForVirtualizedContentCommandlet::Main(const FString& Params)
 	{
 		TRACE_CPUPROFILER_EVENT_SCOPE(SortAllPackages);
 
-		TArray<FString> AllPackages = UE::Virtualization::FindAllPackages();
+		TArray<FString> AllPackages = UE::Virtualization::FindPackages(UE::Virtualization::EFindPackageFlags::None);
 
 		const FString EngineDir = FPaths::EngineDir();
 
diff --git a/Engine/Source/Editor/VirtualizationEditor/Private/CommandletUtils.cpp b/Engine/Source/Editor/VirtualizationEditor/Private/CommandletUtils.cpp
--- a/Engine/Source/Editor/VirtualizationEditor/Private/CommandletUtils.cpp
+++ b/Engine/Source/Editor/VirtualizationEditor/Private/CommandletUtils.cpp
@@ -13,9 +13,9 @@
 namespace UE::Virtualization
 {
 
-TArray<FString> FindAllPackages()
+TArray<FString> FindPackages(EFindPackageFlags Flags)
 {
-	TRACE_CPUPROFILER_EVENT_SCOPE(FindAllPackages);
+	TRACE_CPUPROFILER_EVENT_SCOPE(FindPackages);
 
 	TArray<FString> PackagePaths;
 
@@ -30,7 +30,9 @@ TArray<FString> FindAllPackages()
 
 	const FString EnginePath = FPaths::EngineDir();
 
-	AssetRegistry.Get().EnumerateAllPackages([&PackagePaths, EnginePath](FName PackageName, const FAssetPackageData& PackageData)
+	const bool bFilterEngineContent = Flags & EFindPackageFlags::ExcludeEngineContent;
+
+	AssetRegistry.Get().EnumerateAllPackages([&PackagePaths, EnginePath, bFilterEngineContent](FName PackageName, const FAssetPackageData& PackageData)
 		{
 			FString RelFileName;
 			if (PackageData.Extension != EPackageExtension::Unspecified && PackageData.Extension != EPackageExtension::Custom)
@@ -41,7 +43,7 @@ TArray<FString> FindAllPackages()
 					FString StdFileName = FPaths::CreateStandardFilename(RelFileName);
 				
 					// Now we have the absolute file path we can filter out engine packages
-					if (!StdFileName.StartsWith(EnginePath))
+					if (!bFilterEngineContent || !StdFileName.StartsWith(EnginePath))
 					{
 						PackagePaths.Emplace(MoveTemp(StdFileName));
 					}
@@ -68,7 +70,7 @@ TArray<FString> FindPackagesInDirectory(const FString& DirectoryToSearch)
 	return PackageNames;
 }
 
-TArray<FString> DiscoverPackages(const FString& CmdlineParams)
+TArray<FString> DiscoverPackages(const FString& CmdlineParams, EFindPackageFlags Flags)
 {
 	TRACE_CPUPROFILER_EVENT_SCOPE(DiscoverPackages);
 
@@ -79,7 +81,7 @@ TArray<FString> DiscoverPackages(const FString& CmdlineParams)
 	}
 	else
 	{
-		return FindAllPackages();
+		return FindPackages(Flags);
 	}
 }
 
diff --git a/Engine/Source/Editor/VirtualizationEditor/Private/CommandletUtils.h b/Engine/Source/Editor/VirtualizationEditor/Private/CommandletUtils.h
--- a/Engine/Source/Editor/VirtualizationEditor/Private/CommandletUtils.h
+++ b/Engine/Source/Editor/VirtualizationEditor/Private/CommandletUtils.h
@@ -11,8 +11,19 @@ namespace UE { class FPackageTrailer; }
 namespace UE::Virtualization
 {
 
-/** Parse all of the active mount points and find all packages */
-TArray<FString> FindAllPackages();
+/** Used to customize package discovery behavior */
+enum EFindPackageFlags
+{
+	None = 0,
+	/** Do not return packages mounted to the engine or engine plugins */
+	ExcludeEngineContent = 1 << 0
+};
+
+/** 
+ * Parse all of the active mount points for the current project and returns the
+ * packages found.
+ */
+TArray<FString> FindPackages(EFindPackageFlags Flags);
 
 /** Finds all of the packages under a given directory including its subdirectories */
 TArray<FString> FindPackagesInDirectory(const FString& DirectoryToSearch);
@@ -28,7 +39,7 @@ TArray<FString> FindPackagesInDirectory(const FString& DirectoryToSearch);
  * @param CmdlineParams A string containing the command line
  * @return An array with the file path for each package found
  */
-TArray<FString> DiscoverPackages(const FString& CmdlineParams);
+TArray<FString> DiscoverPackages(const FString& CmdlineParams, EFindPackageFlags Flags);
 
 /** Returns a combined list of unique virtualized payload ids from the given list of packages */
 TArray<FIoHash> FindVirtualizedPayloads(const TArray<FString>& PackageNames);
diff --git a/Engine/Source/Editor/VirtualizationEditor/Private/GeneratePayloadManifestCommandlet.cpp b/Engine/Source/Editor/VirtualizationEditor/Private/GeneratePayloadManifestCommandlet.cpp
--- a/Engine/Source/Editor/VirtualizationEditor/Private/GeneratePayloadManifestCommandlet.cpp
+++ b/Engine/Source/Editor/VirtualizationEditor/Private/GeneratePayloadManifestCommandlet.cpp
@@ -18,7 +18,7 @@ int32 UGeneratePayloadManifestCommandlet::Main(const FString& Params)
 {
 	TRACE_CPUPROFILER_EVENT_SCOPE(UGeneratePayloadManifestCommandlet);
 
-	TArray<FString> PackageNames = UE::Virtualization::DiscoverPackages(Params);
+	TArray<FString> PackageNames = UE::Virtualization::DiscoverPackages(Params, UE::Virtualization::EFindPackageFlags::ExcludeEngineContent);
 
 	UE_LOG(LogVirtualization, Display, TEXT("Found %d files to look in"), PackageNames.Num());
 
diff --git a/Engine/Source/Editor/VirtualizationEditor/Private/PrecachePayloadsCommandlet.cpp b/Engine/Source/Editor/VirtualizationEditor/Private/PrecachePayloadsCommandlet.cpp
--- a/Engine/Source/Editor/VirtualizationEditor/Private/PrecachePayloadsCommandlet.cpp
+++ b/Engine/Source/Editor/VirtualizationEditor/Private/PrecachePayloadsCommandlet.cpp
@@ -114,7 +114,7 @@ int32 UPrecachePayloadsCommandlet::Main(const FString& Params)
 {
 	TRACE_CPUPROFILER_EVENT_SCOPE(UPrecachePayloadsCommandlet);
 
-	TArray<FString> PackageNames = UE::Virtualization::FindAllPackages();
+	TArray<FString> PackageNames = UE::Virtualization::FindPackages(UE::Virtualization::EFindPackageFlags::ExcludeEngineContent);
 
 	UE_LOG(LogVirtualization, Display, TEXT("Found %d packages"), PackageNames.Num());
 
diff --git a/Engine/Source/Editor/VirtualizationEditor/Private/RehydrateProjectCommandlet.cpp b/Engine/Source/Editor/VirtualizationEditor/Private/RehydrateProjectCommandlet.cpp
--- a/Engine/Source/Editor/VirtualizationEditor/Private/RehydrateProjectCommandlet.cpp
+++ b/Engine/Source/Editor/VirtualizationEditor/Private/RehydrateProjectCommandlet.cpp
@@ -27,7 +27,7 @@ int32 URehydrateProjectCommandlet::Main(const FString& Params)
 	}
 
 	UE_LOG(LogVirtualization, Display, TEXT("Scanning for project packages..."));
-	TArray<FString> Packages = UE::Virtualization::FindAllPackages();
+	TArray<FString> Packages = UE::Virtualization::FindPackages(UE::Virtualization::EFindPackageFlags::ExcludeEngineContent);
 
 	UE_LOG(LogVirtualization, Display, TEXT("Rehydrating packages..."));
 
diff --git a/Engine/Source/Editor/VirtualizationEditor/Private/ValidateVirtualizedContentCommandlet.cpp b/Engine/Source/Editor/VirtualizationEditor/Private/ValidateVirtualizedContentCommandlet.cpp
--- a/Engine/Source/Editor/VirtualizationEditor/Private/ValidateVirtualizedContentCommandlet.cpp
+++ b/Engine/Source/Editor/VirtualizationEditor/Private/ValidateVirtualizedContentCommandlet.cpp
@@ -18,7 +18,7 @@ int32 UValidateVirtualizedContentCommandlet::Main(const FString& Params)
 	TRACE_CPUPROFILER_EVENT_SCOPE(UValidateVirtualizedContentCommandlet);
 
 	UE_LOG(LogVirtualization, Display, TEXT("Finding packages in the project..."));
-	TArray<FString> PackagePaths = FindAllPackages();
+	TArray<FString> PackagePaths = FindPackages(EFindPackageFlags::ExcludeEngineContent);
 	UE_LOG(LogVirtualization, Display, TEXT("Found %d package(s)"), PackagePaths.Num());
 
 	TMap<FString, UE::FPackageTrailer> Packages;
diff --git a/Engine/Source/Editor/VirtualizationEditor/Private/VirtualizeProjectCommandlet.cpp b/Engine/Source/Editor/VirtualizationEditor/Private/VirtualizeProjectCommandlet.cpp
--- a/Engine/Source/Editor/VirtualizationEditor/Private/VirtualizeProjectCommandlet.cpp
+++ b/Engine/Source/Editor/VirtualizationEditor/Private/VirtualizeProjectCommandlet.cpp
@@ -27,7 +27,7 @@ int32 UVirtualizeProjectCommandlet::Main(const FString& Params)
 	}
 
 	UE_LOG(LogVirtualization, Display, TEXT("Scanning for project packages..."));
-	TArray<FString> Packages = UE::Virtualization::FindAllPackages();
+	TArray<FString> Packages = UE::Virtualization::FindPackages(UE::Virtualization::EFindPackageFlags::ExcludeEngineContent);
 
 	UE_LOG(LogVirtualization, Display, TEXT("Virtualizing packages..."));
 
