Improve VA commandlet documentation to avoid problems when people invoke them via powershell.

#rnx

- When a tester tried '-run=VirtualizationEditor.???' via powershell a spece was added so the editor got '-run=VirtualizationEditor .???' which doesn't parse correctly.
- Change the documentation to suggest using '-run="VirtualizationEditor.???"' so that the original text is preserved.

diff --git a/Engine/Source/Editor/VirtualizationEditor/Private/CheckForVirtualizedContentCommandlet.h b/Engine/Source/Editor/VirtualizationEditor/Private/CheckForVirtualizedContentCommandlet.h
--- a/Engine/Source/Editor/VirtualizationEditor/Private/CheckForVirtualizedContentCommandlet.h
+++ b/Engine/Source/Editor/VirtualizationEditor/Private/CheckForVirtualizedContentCommandlet.h
@@ -25,7 +25,7 @@
  *
  * Because the commandlet is the VirtualizationEditor module it needs to be invoked
  * with the command line:
- * -run=VirtualizationEditor.CheckForVirtualizedContent
+ * -run="VirtualizationEditor.CheckForVirtualizedContent"
  *
  * Followed by
  */
diff --git a/Engine/Source/Editor/VirtualizationEditor/Private/GeneratePayloadManifestCommandlet.h b/Engine/Source/Editor/VirtualizationEditor/Private/GeneratePayloadManifestCommandlet.h
--- a/Engine/Source/Editor/VirtualizationEditor/Private/GeneratePayloadManifestCommandlet.h
+++ b/Engine/Source/Editor/VirtualizationEditor/Private/GeneratePayloadManifestCommandlet.h
@@ -15,7 +15,7 @@
  * 
  * Because the commmandlet is the VirtualizationEditor module it needs to be invoked 
  * with the command line:
- * -run=VirtualizationEditor.GeneratePayloadManifest
+ * -run="VirtualizationEditor.GeneratePayloadManifest"
  */
 UCLASS()
 class UGeneratePayloadManifestCommandlet
diff --git a/Engine/Source/Editor/VirtualizationEditor/Private/PrecachePayloadsCommandlet.h b/Engine/Source/Editor/VirtualizationEditor/Private/PrecachePayloadsCommandlet.h
--- a/Engine/Source/Editor/VirtualizationEditor/Private/PrecachePayloadsCommandlet.h
+++ b/Engine/Source/Editor/VirtualizationEditor/Private/PrecachePayloadsCommandlet.h
@@ -17,7 +17,7 @@
  *
  * Because the commandlet is the VirtualizationEditor module it needs to be invoked
  * with the command line:
- * -run=VirtualizationEditor.PrecachePayloads
+ * -run="VirtualizationEditor.PrecachePayloads"
  */
 UCLASS()
 class UPrecachePayloadsCommandlet
diff --git a/Engine/Source/Editor/VirtualizationEditor/Private/RehydrateProjectCommandlet.h b/Engine/Source/Editor/VirtualizationEditor/Private/RehydrateProjectCommandlet.h
--- a/Engine/Source/Editor/VirtualizationEditor/Private/RehydrateProjectCommandlet.h
+++ b/Engine/Source/Editor/VirtualizationEditor/Private/RehydrateProjectCommandlet.h
@@ -13,7 +13,7 @@
  *
  * Because the commmandlet is the VirtualizationEditor module it needs to be invoked
  * with the command line:
- * -run=VirtualizationEditor.RehydrateProject
+ * -run="VirtualizationEditor.RehydrateProject"
  */
 UCLASS()
 class URehydrateProjectCommandlet
diff --git a/Engine/Source/Editor/VirtualizationEditor/Private/ValidateVirtualizedContentCommandlet.h b/Engine/Source/Editor/VirtualizationEditor/Private/ValidateVirtualizedContentCommandlet.h
--- a/Engine/Source/Editor/VirtualizationEditor/Private/ValidateVirtualizedContentCommandlet.h
+++ b/Engine/Source/Editor/VirtualizationEditor/Private/ValidateVirtualizedContentCommandlet.h
@@ -16,7 +16,7 @@
  * 
  * Because the commandlet is the VirtualizationEditor module it needs to be invoked
  * with the command line:
- * -run=VirtualizationEditor.ValidateVirtualizedContent
+ * -run="VirtualizationEditor.ValidateVirtualizedContent"
  */
 UCLASS()
 class UValidateVirtualizedContentCommandlet
diff --git a/Engine/Source/Editor/VirtualizationEditor/Private/VirtualizeProjectCommandlet.h b/Engine/Source/Editor/VirtualizationEditor/Private/VirtualizeProjectCommandlet.h
--- a/Engine/Source/Editor/VirtualizationEditor/Private/VirtualizeProjectCommandlet.h
+++ b/Engine/Source/Editor/VirtualizationEditor/Private/VirtualizeProjectCommandlet.h
@@ -8,12 +8,12 @@
 #include "VirtualizeProjectCommandlet.generated.h"
 
 /**
- * Finds all packages in the project and attempts to virtualize them. If revision control
+ * Finds all packages in the project and attempts to virtualize their content. If revision control
  * is enabled then the commandlet will attempt to checkout the packages that need modification.
  *
- * Because the commmandlet is the VirtualizationEditor module it needs to be invoked
+ * Because the commandlet is the VirtualizationEditor module it needs to be invoked
  * with the command line:
- * -run=VirtualizationEditor.VirtualizeProject
+ * -run="VirtualizationEditor.VirtualizeProject"
  */
 UCLASS()
 class UVirtualizeProjectCommandlet
