Fix a typo in the VA commandlet documentation.

#rnx

diff --git a/Engine/Source/Editor/VirtualizationEditor/Private/RehydrateProjectCommandlet.h b/Engine/Source/Editor/VirtualizationEditor/Private/RehydrateProjectCommandlet.h
--- a/Engine/Source/Editor/VirtualizationEditor/Private/RehydrateProjectCommandlet.h
+++ b/Engine/Source/Editor/VirtualizationEditor/Private/RehydrateProjectCommandlet.h
@@ -11,7 +11,7 @@
  * Finds all packages in the project and attempts to rehydrate them. If revision control
  * is enabled then the commandlet will attempt to checkout the packages that need modification.
  *
- * Because the commmandlet is the VirtualizationEditor module it needs to be invoked
+ * Because the commandlet is the VirtualizationEditor module it needs to be invoked
  * with the command line:
  * -run="VirtualizationEditor.RehydrateProject"
  */
