Improve the VA PrecachePayloads commandlet

#rnx

- The previous default iteration mode could saturdate the worker threads, which would prevent the Zen DDC backend from running clean up code resulting in very large memory spikes.
- We now launch a limited number of tasks that process a limited number of payloads before completing. Once a task is finsihed we issue another until there is no more work to do.
- This seems to strike a good balance of keeping the network transfer at peak speed but allowing other systems time on the cpu.
- Removed all the experimental iteration modes and going all in on this solution, which makes the code easier to parse.

diff --git a/Engine/Source/Editor/VirtualizationEditor/Private/PrecachePayloadsCommandlet.cpp b/Engine/Source/Editor/VirtualizationEditor/Private/PrecachePayloadsCommandlet.cpp
--- a/Engine/Source/Editor/VirtualizationEditor/Private/PrecachePayloadsCommandlet.cpp
+++ b/Engine/Source/Editor/VirtualizationEditor/Private/PrecachePayloadsCommandlet.cpp
@@ -8,16 +8,6 @@
 #include "UObject/PackageTrailer.h"
 #include "Virtualization/VirtualizationSystem.h"
 
-// 0 - No threading
-// 1 - ParallelFor
-// 2 - Split the work into equal sized batches across all cores
-// 3 - Use a limited number of worker tasks
-// 4 - Single batch request 
-// 5 - Use a limited number of worker tasks with batch requests
-
-#define UE_ITERATION_METHOD 5
-#define UE_VALIDATE_PAYLOADS 0
-
 namespace 
 {
 
@@ -46,8 +36,6 @@ public:
 	{
 		TRACE_CPUPROFILER_EVENT_SCOPE(FWorkQueue::GetJob);
 
-		FScopeLock _(&JobsCS);
-
 		if (!Jobs.IsEmpty())
 		{
 			return Jobs.Pop(false);
@@ -55,39 +43,24 @@ public:
 		else
 		{
 			return FJob();
-		}	
+		}
+	}
+
+	bool IsEmpty() const
+	{
+		return Jobs.IsEmpty();
 	}
 
+	int32 Num() const
+	{
+		return Jobs.Num();
+	}
 
+private:
 	TArray<FIoHash> Work;
 	TArray<FJob> Jobs;
-
-	FCriticalSection JobsCS;
 };
 
-/** 
- * Utility to wait for the given tasks to complete and return early if the tasks are not completed
- * within the given timespan.
- * This is used isntead of UE::Tasks::Wait, because that call can execute the remaining tasks if not
- * already being processed, which in most scenarios makes sense, but since we are running in a commandlet
- * and we know nothing else is running, we don't need the game thread processing work as well. 
- */
-bool WaitOnTasks(const TArray<UE::Tasks::FTask>& Tasks, FTimespan InTimeout)
-{
-	UE::FTimeout Timeout{ InTimeout };
-
-	FSharedEventRef CompletionEvent;
-
-	UE::Tasks::FTask WaitingTask = UE::Tasks::Launch(UE_SOURCE_LOCATION,
-		[CompletionEvent]
-		{
-			CompletionEvent->Trigger();
-		},
-		Tasks, LowLevelTasks::ETaskPriority::Default, UE::Tasks::EExtendedTaskPriority::Inline);
-
-	return CompletionEvent->Wait((uint32)FMath::Clamp<int64>(Timeout.GetRemainingTime().GetTicks() / ETimespan::TicksPerMillisecond, 0, MAX_uint32));
-}
-
 /** Utility to turn an array of FIoHash into an array of FPullRequest */
 TArray<UE::Virtualization::FPullRequest> ToRequestArray(TConstArrayView<FIoHash> IdentifierArray)
 {
@@ -130,213 +103,75 @@ int32 UPrecachePayloadsCommandlet::Main(const FString& Params)
 	UE_LOG(LogVirtualization, Display, TEXT("Precaching payloads..."));
 
 	UE::Virtualization::IVirtualizationSystem& System = UE::Virtualization::IVirtualizationSystem::Get();
-
-#if UE_ITERATION_METHOD == 0
-	{
-		TRACE_CPUPROFILER_EVENT_SCOPE(Precache_SingleRequests);
-
-		for (const FIoHash& PayloadId : PayloadIds)
-		{
-			FCompressedBuffer Payload = System.PullData(PayloadId);
-			if (Payload.IsNull())
-			{
-				UE_LOG(LogVirtualization, Error, TEXT("Failed to precache payload '%s'"), *LexToString(PayloadId));
-			}
-		}
-	}
-#elif UE_ITERATION_METHOD == 1
-	{
-		TRACE_CPUPROFILER_EVENT_SCOPE(Precache_ParallelFor);
-
-		ParallelFor(PayloadIds.Num(), [&PayloadIds, &System](int32 Index)
-			{
-				const FIoHash& PayloadId = PayloadIds[Index];
-
-				FCompressedBuffer Payload = System.PullData(PayloadId);
-				if (Payload.IsNull())
-				{
-					UE_LOG(LogVirtualization, Error, TEXT("Failed to precache payload '%s'"), *LexToString(PayloadId));
-				}
-			});
-	}
-#elif UE_ITERATION_METHOD == 2
 	{
 		TRACE_CPUPROFILER_EVENT_SCOPE(Precache_ThreadedBatches);
 
-		int32 NumBatches = FPlatformMisc::NumberOfCores();
-
-		TArray<UE::Tasks::FTask> Tasks;
-		Tasks.Reserve(NumBatches);
-		
 		std::atomic<int32> NumCompletedPayloads = 0;
-		const int32 BatchSize = FMath::DivideAndRoundUp<int32>(PayloadIds.Num(), NumBatches);
-		NumBatches -= (((NumBatches * BatchSize) - PayloadIds.Num()) / BatchSize);
-
-		check(NumBatches * BatchSize >= PayloadIds.Num());
-
-		for (int32 BatchIndex = 0; BatchIndex < NumBatches; ++BatchIndex)
-		{
-			const int32 BatchStart = BatchIndex * BatchSize;
-			const int32 BatchEnd = FMath::Min((BatchIndex + 1) * BatchSize, PayloadIds.Num());
-
-			TArrayView<const FIoHash> Slice = MakeArrayView(&PayloadIds[BatchStart], BatchEnd - BatchStart);
-			
-			Tasks.Emplace(UE::Tasks::Launch(UE_SOURCE_LOCATION, 
-				[Slice, &System, &NumCompletedPayloads]()
-				{
-					for (const FIoHash& PayloadId : Slice)
-					{
-						FCompressedBuffer Payload = System.PullData(PayloadId);
-						if (Payload.IsNull())
-						{
-							UE_LOG(LogVirtualization, Error, TEXT("Failed to precache payload '%s'"), *LexToString(PayloadId));
-						}
 
-						NumCompletedPayloads++;
-					}
-				}));
-		}	
+		const int32 NumPayloads = PayloadIds.Num();
+		const int32 BatchSize = 64;
 
-		while (!WaitOnTasks(Tasks, FTimespan::FromSeconds(30.0)))
-		{
-			const int32 CurrentCompletedPayloads = NumCompletedPayloads;
+		// This seems to be the sweet spot when it comes to our internal infrastructure, so use that as the default.
+		const int32 MaxConcurrentTasks = 16;
 
-			const float Progress = ((float)CurrentCompletedPayloads / (float)PayloadIds.Num()) * 100.0f;
+		// We always want to leave at least one foreground worker free to avoid saturation. If we issue too many
+		// concurrent task then we can potentially cause the DDC/Zen to be unable to run clean up tasks for long 
+		// periods of times which can cause quite high memory spikes.
+		const int32 ConcurrentTasks = FMath::Min(MaxConcurrentTasks, FTaskGraphInterface::Get().GetNumWorkerThreads() - 1);
 
-			UE_LOG(LogVirtualization, Display, TEXT("Cached %d/%d (%.1f%%)"), CurrentCompletedPayloads, PayloadIds.Num(), Progress);
-		}
-	}
-#elif UE_ITERATION_METHOD == 3
-	{
-		TRACE_CPUPROFILER_EVENT_SCOPE(Precache_ThreadedWorkerBatches);
-
-		const int NumWorkers = FPlatformMisc::NumberOfCores();
-		const int BatchSize = 64;
-		const int NumPayloads = PayloadIds.Num();
-
-		std::atomic<int32> NumCompletedPayloads = 0;
+		UE_LOG(LogVirtualization, Display, TEXT("Will run up to %d precache tasks concurrently"), ConcurrentTasks);
 
 		FWorkQueue WorkQueue(MoveTemp(PayloadIds), BatchSize);
 
-		TArray<UE::Tasks::FTask> Tasks;
-		Tasks.Reserve(NumWorkers);
-
-		for (int32 Index = 0; Index < NumWorkers; ++Index)
-		{
-			Tasks.Emplace(UE::Tasks::Launch(UE_SOURCE_LOCATION,
-				[&WorkQueue, &System, &NumCompletedPayloads]()
-				{
-					while(true)
-					{
-						FWorkQueue::FJob Job = WorkQueue.GetJob();
-
-						if (Job.IsEmpty())
-						{
-							return;
-						}
-
-						for (const FIoHash& PayloadId : Job)
-						{
-							FCompressedBuffer Payload = System.PullData(PayloadId);
-							if (Payload.IsNull())
-							{
-								UE_LOG(LogVirtualization, Error, TEXT("Failed to precache payload '%s'"), *LexToString(PayloadId));
-							}
+		UE::Tasks::FTaskEvent Event(UE_SOURCE_LOCATION);
 
-							NumCompletedPayloads++;
-						}
-					}
-				}));
-		}
+		std::atomic<int32> NumTasks = 0;
+		double LogTimer = FPlatformTime::Seconds();
 
-		while (!WaitOnTasks(Tasks, FTimespan::FromSeconds(30.0)))
+		while (NumTasks != 0 || !WorkQueue.IsEmpty())
 		{
-			const int32 CurrentCompletedPayloads = NumCompletedPayloads;
-
-			const float Progress = ((float)CurrentCompletedPayloads / (float)NumPayloads) * 100.0f;
-
-			UE_LOG(LogVirtualization, Display, TEXT("Cached %d/%d (%.1f%%)"), CurrentCompletedPayloads, NumPayloads, Progress);
-		}
-	}
-#elif UE_ITERATION_METHOD == 4
-	{
-		TRACE_CPUPROFILER_EVENT_SCOPE(Precache_SingleBatchRequest);
-
-		TArray<UE::Virtualization::FPullRequest> Requests = ToRequestArray(PayloadIds);
-	
-		System.PullData(Requests);
-
-		for (const UE::Virtualization::FPullRequest& Request : Requests)
-		{
-			if (!Request.IsSuccess())
+			int32 NumTaskAllowed = ConcurrentTasks - NumTasks;
+			while (NumTaskAllowed > 0 && !WorkQueue.IsEmpty())
 			{
-				UE_LOG(LogVirtualization, Error, TEXT("Failed to precache payload '%s'"), *LexToString(Request.GetIdentifier()));
-			}
-		}
-	}
-#elif UE_ITERATION_METHOD == 5
-	{
-		TRACE_CPUPROFILER_EVENT_SCOPE(Precache_ThreadedBatchRequest);
+				NumTasks++;
 
-		const int NumWorkers = FPlatformMisc::NumberOfCores();
-		const int BatchSize = 64;
-		const int NumPayloads = PayloadIds.Num();
-
-		std::atomic<int32> NumCompletedPayloads = 0;
-
-		FWorkQueue WorkQueue(MoveTemp(PayloadIds), BatchSize);
-
-		TArray<UE::Tasks::FTask> Tasks;
-		Tasks.Reserve(NumWorkers);
-
-		for (int32 Index = 0; Index < NumWorkers; ++Index)
-		{
-			Tasks.Emplace(UE::Tasks::Launch(UE_SOURCE_LOCATION,
-				[&WorkQueue, &System, &NumCompletedPayloads]()
-				{
-					while (true)
+				UE::Tasks::Launch(UE_SOURCE_LOCATION,
+					[Job = WorkQueue.GetJob(), &System, &NumCompletedPayloads, &NumTasks, &Event]()
 					{
-						FWorkQueue::FJob Job = WorkQueue.GetJob();
-
-						if (Job.IsEmpty())
-						{
-							return;
-						}
-
 						TArray<UE::Virtualization::FPullRequest> Requests = ToRequestArray(Job);
+
 						if (!System.PullData(Requests))
 						{
-							for (const FIoHash& PayloadId : Job)
+							for (const UE::Virtualization::FPullRequest& Request : Requests)
 							{
-								FCompressedBuffer Payload = System.PullData(PayloadId);
-								if (Payload.IsNull())
-								{
-									UE_LOG(LogVirtualization, Error, TEXT("%s: Failed to precache payload"), *LexToString(PayloadId));
-								}
-#if UE_VALIDATE_PAYLOADS
-								else if (Payload.GetRawHash() != PayloadId)
+								if (!Request.IsSuccess())
 								{
-									UE_LOG(LogVirtualization, Error, TEXT("%s: DAta was corrupted"), *LexToString(PayloadId));
+									UE_LOG(LogVirtualization, Error, TEXT("%s: Failed to precache payload"), *LexToString(Request.GetIdentifier()));
 								}
-#endif //UE_VALIDATE_PAYLOADS
 							}
 						}
 
-						NumCompletedPayloads += Job.Num();
-					}
-				}));
-		}
+						NumCompletedPayloads += Requests.Num();
+						--NumTasks;
 
-		while (!WaitOnTasks(Tasks, FTimespan::FromSeconds(30.0)))
-		{
-			const int32 CurrentCompletedPayloads = NumCompletedPayloads;
+						Event.Trigger();
+					});
 
-			const float Progress = ((float)CurrentCompletedPayloads / (float)NumPayloads) * 100.0f;
+				--NumTaskAllowed;
+			}
 
-			UE_LOG(LogVirtualization, Display, TEXT("Cached %d/%d (%.1f%%)"), CurrentCompletedPayloads, NumPayloads, Progress);
+			Event.Wait(FTimespan::FromSeconds(30.0));
+
+			if (FPlatformTime::Seconds() - LogTimer >= 30.0)
+			{
+				const int32 CurrentCompletedPayloads = NumCompletedPayloads;
+				const float Progress = ((float)CurrentCompletedPayloads / (float)NumPayloads) * 100.0f;
+				UE_LOG(LogVirtualization, Display, TEXT("Cached %d/%d (%.1f%%)"), CurrentCompletedPayloads, NumPayloads, Progress);
+
+				LogTimer = FPlatformTime::Seconds();
+			}
 		}
 	}
-#endif // UE_ITERATION_METHOD
 
 	UE_LOG(LogVirtualization, Display, TEXT("Precaching complete!"));
 
@@ -344,6 +179,3 @@ int32 UPrecachePayloadsCommandlet::Main(const FString& Params)
 	
 	return  0;
 }
-
-#undef UE_VALIDATE_PAYLOADS
-#undef UE_ITERATION_METHOD
