Move some comon functionality from UPrecachePayloadsCommandlet to the general utilities so that other commandlets can use it.

#rnx

diff --git a/Engine/Source/Editor/VirtualizationEditor/Private/CommandletUtils.cpp b/Engine/Source/Editor/VirtualizationEditor/Private/CommandletUtils.cpp
--- a/Engine/Source/Editor/VirtualizationEditor/Private/CommandletUtils.cpp
+++ b/Engine/Source/Editor/VirtualizationEditor/Private/CommandletUtils.cpp
@@ -6,9 +6,11 @@
 #include "AssetRegistry/AssetRegistryModule.h"
 #include "Async/ParallelFor.h"
 #include "Containers/Set.h"
+#include "IO/IoHash.h"
 #include "Misc/PackageName.h"
 #include "Misc/Paths.h"
 #include "UObject/PackageTrailer.h"
+#include "Virtualization/VirtualizationSystem.h"
 
 namespace UE::Virtualization
 {
@@ -151,10 +153,64 @@ void FindVirtualizedPayloadsAndTrailers(const TArray<FString>& PackagePaths, TMa
 		OutPayloads.Append(Context.Payloads);
 	}
 
+	OutPayloads.CompactStable();
+}
+
+TArray<FPullRequest> ToRequestArray(TConstArrayView<FIoHash> IdentifierArray)
+{
+	TArray<FPullRequest> Requests;
+	Requests.Reserve(IdentifierArray.Num());
+
+	for (const FIoHash& Id : IdentifierArray)
 	{
-		TRACE_CPUPROFILER_EVENT_SCOPE(CompactStable);
-		OutPayloads.CompactStable();
+		Requests.Emplace(FPullRequest(Id));
 	}
+
+	return Requests;
+}
+
+FWorkQueue::FWorkQueue(const TArray<FIoHash>& InWork, int32 JobSize)
+	: Work(InWork)
+{
+	CreateJobs(JobSize);
+}
+
+FWorkQueue::FWorkQueue(TArray<FIoHash>&& InWork, int32 JobSize)
+	: Work(InWork)
+{
+	CreateJobs(JobSize);
+}
+
+void FWorkQueue::CreateJobs(int32 JobSize)
+{
+	const int32 NumJobs = FMath::DivideAndRoundUp<int32>(Work.Num(), JobSize);
+	Jobs.Reserve(NumJobs);
+
+	for (int32 JobIndex = 0; JobIndex < NumJobs; ++JobIndex)
+	{
+		const int32 JobStart = JobIndex * JobSize;
+		const int32 JobEnd = FMath::Min((JobIndex + 1) * JobSize, Work.Num());
+
+		FJob Job = MakeArrayView(&Work[JobStart], JobEnd - JobStart);
+		Jobs.Add(Job);
+	}
+}
+
+FWorkQueue::FJob FWorkQueue::GetJob()
+{
+	if (!Jobs.IsEmpty())
+	{
+		return Jobs.Pop(false);
+	}
+	else
+	{
+		return FJob();
+	}
+}
+
+bool FWorkQueue::IsEmpty() const
+{
+	return Jobs.IsEmpty();
 }
 
 } //namespace UE::Virtualization
diff --git a/Engine/Source/Editor/VirtualizationEditor/Private/CommandletUtils.h b/Engine/Source/Editor/VirtualizationEditor/Private/CommandletUtils.h
--- a/Engine/Source/Editor/VirtualizationEditor/Private/CommandletUtils.h
+++ b/Engine/Source/Editor/VirtualizationEditor/Private/CommandletUtils.h
@@ -4,11 +4,12 @@
 
 #include "Containers/Array.h"
 #include "Containers/UnrealString.h"
-#include "IO/IoHash.h"
 
 namespace UE { class FPackageTrailer; }
+namespace UE::Virtualization { struct FPullRequest; }
+struct FIoHash;
 
-namespace UE::Virtualization
+namespace UE::Virtualization // Utility functions
 {
 
 /** Used to customize package discovery behavior */
@@ -55,4 +56,48 @@ TArray<FIoHash> FindVirtualizedPayloads(const TArray<FString>& PackageNames);
  */
 void FindVirtualizedPayloadsAndTrailers(const TArray<FString>& PackagePaths, TMap<FString, UE::FPackageTrailer>& OutPackages, TSet<FIoHash>& OutPayloads);
 
+/** Utility to turn an array of FIoHash into an array of FPullRequest */
+TArray<FPullRequest> ToRequestArray(TConstArrayView<FIoHash> IdentifierArray);
+
+} //namespace UE::Virtualization
+
+
+namespace UE::Virtualization // Utility classes
+{
+
+/** 
+ * Takes a list of FIoHashs and breaks them into smaller jobs based on the JobSize.
+ * These jobs (in the form of an TArrayView) can then be requested one at a time
+ * until there are no more jobs left.
+ * The intended use is to make it easier to divide up the FIoHash list into smaller
+ * units to be pushed to the taskgraph system.
+ */
+class FWorkQueue
+{
+public:
+	using FJob = TArrayView<const FIoHash>;
+
+	FWorkQueue(const TArray<FIoHash>& InWork, int32 JobSize);
+	FWorkQueue(TArray<FIoHash>&& InWork, int32 JobSize);
+	~FWorkQueue() = default;
+
+	/** 
+	 * Returns a FJob containing FIoHashes to be processed. 
+	 * If there is no more work to be done then empty jobs will be returned 
+	 * 
+	 * NOTE: That FJob is only valid for the lifespan of the FWorkQueue
+	 */
+	FJob GetJob();
+
+	/** Returns true if there is no more work to done */
+	bool IsEmpty() const;
+
+private:
+
+	void CreateJobs(int32 JobSize);
+
+	TArray<FIoHash> Work;
+	TArray<FJob> Jobs;
+};
+
 } //namespace UE::Virtualization
diff --git a/Engine/Source/Editor/VirtualizationEditor/Private/PrecachePayloadsCommandlet.cpp b/Engine/Source/Editor/VirtualizationEditor/Private/PrecachePayloadsCommandlet.cpp
--- a/Engine/Source/Editor/VirtualizationEditor/Private/PrecachePayloadsCommandlet.cpp
+++ b/Engine/Source/Editor/VirtualizationEditor/Private/PrecachePayloadsCommandlet.cpp
@@ -8,76 +8,6 @@
 #include "UObject/PackageTrailer.h"
 #include "Virtualization/VirtualizationSystem.h"
 
-namespace 
-{
-
-class FWorkQueue
-{
-public:
-	using FJob = TArrayView<const FIoHash>;
-
-	FWorkQueue(TArray<FIoHash>&& InWork, int32 JobSize)
-		: Work(InWork)
-	{
-		const int32 NumJobs = FMath::DivideAndRoundUp<int32>(Work.Num(), JobSize);
-		Jobs.Reserve(NumJobs);
-
-		for (int32 JobIndex = 0; JobIndex < NumJobs; ++JobIndex)
-		{
-			const int32 JobStart = JobIndex * JobSize;
-			const int32 JobEnd = FMath::Min((JobIndex + 1) * JobSize, Work.Num());
-
-			FJob Job = MakeArrayView(&Work[JobStart], JobEnd - JobStart);
-			Jobs.Add(Job);
-		}
-	}
-
-	FJob GetJob()
-	{
-		TRACE_CPUPROFILER_EVENT_SCOPE(FWorkQueue::GetJob);
-
-		if (!Jobs.IsEmpty())
-		{
-			return Jobs.Pop(false);
-		}
-		else
-		{
-			return FJob();
-		}
-	}
-
-	bool IsEmpty() const
-	{
-		return Jobs.IsEmpty();
-	}
-
-	int32 Num() const
-	{
-		return Jobs.Num();
-	}
-
-private:
-	TArray<FIoHash> Work;
-	TArray<FJob> Jobs;
-};
-
-/** Utility to turn an array of FIoHash into an array of FPullRequest */
-TArray<UE::Virtualization::FPullRequest> ToRequestArray(TConstArrayView<FIoHash> IdentifierArray)
-{
-	TArray<UE::Virtualization::FPullRequest> Requests;
-	Requests.Reserve(IdentifierArray.Num());
-
-	for (const FIoHash& Id : IdentifierArray)
-	{
-		Requests.Emplace(Id);
-	}
-
-	return Requests;
-}
-
-} // namespace
-
-
 UPrecachePayloadsCommandlet::UPrecachePayloadsCommandlet(const FObjectInitializer& ObjectInitializer)
 	: Super(ObjectInitializer)
 {
@@ -121,7 +51,7 @@ int32 UPrecachePayloadsCommandlet::Main(const FString& Params)
 
 		UE_LOG(LogVirtualization, Display, TEXT("Will run up to %d precache tasks concurrently"), ConcurrentTasks);
 
-		FWorkQueue WorkQueue(MoveTemp(PayloadIds), BatchSize);
+		UE::Virtualization::FWorkQueue WorkQueue(MoveTemp(PayloadIds), BatchSize);
 
 		UE::Tasks::FTaskEvent Event(UE_SOURCE_LOCATION);
 
@@ -138,7 +68,7 @@ int32 UPrecachePayloadsCommandlet::Main(const FString& Params)
 				UE::Tasks::Launch(UE_SOURCE_LOCATION,
 					[Job = WorkQueue.GetJob(), &System, &NumCompletedPayloads, &NumTasks, &Event]()
 					{
-						TArray<UE::Virtualization::FPullRequest> Requests = ToRequestArray(Job);
+						TArray<UE::Virtualization::FPullRequest> Requests = UE::Virtualization::ToRequestArray(Job);
 
 						if (!System.PullData(Requests))
 						{
