Factor out the multithreaded pull system from PrecachePayloadsCommandlet to the general utilities so that other VA commandlets can make use of it.

#rnx

diff --git a/Engine/Source/Editor/VirtualizationEditor/Private/CommandletUtils.cpp b/Engine/Source/Editor/VirtualizationEditor/Private/CommandletUtils.cpp
--- a/Engine/Source/Editor/VirtualizationEditor/Private/CommandletUtils.cpp
+++ b/Engine/Source/Editor/VirtualizationEditor/Private/CommandletUtils.cpp
@@ -9,12 +9,91 @@
 #include "IO/IoHash.h"
 #include "Misc/PackageName.h"
 #include "Misc/Paths.h"
+#include "Tasks/Task.h"	
 #include "UObject/PackageTrailer.h"
 #include "Virtualization/VirtualizationSystem.h"
 
 namespace UE::Virtualization
 {
 
+bool PullPayloadsThreaded(TConstArrayView<FIoHash> PayloadIds, int32 BatchSize, const TCHAR* ProgressString, TUniqueFunction<void(const FPullRequest& Response)>&& Callback)
+{
+	TRACE_CPUPROFILER_EVENT_SCOPE(PullPayloadsThreaded);
+	
+	IVirtualizationSystem& System = IVirtualizationSystem::Get();
+
+	std::atomic<int32> NumCompletedPayloads = 0;
+
+	const int32 NumPayloads = PayloadIds.Num();
+
+	// This seems to be the sweet spot when it comes to our internal infrastructure, so use that as the default.
+	const int32 MaxConcurrentTasks = 16;
+
+	// We always want to leave at least one foreground worker free to avoid saturation. If we issue too many
+	// concurrent task then we can potentially cause the DDC/Zen to be unable to run clean up tasks for long 
+	// periods of times which can cause quite high memory spikes.
+	const int32 ConcurrentTasks = FMath::Min(MaxConcurrentTasks, FTaskGraphInterface::Get().GetNumWorkerThreads() - 1);
+
+	UE_LOG(LogVirtualization, Display, TEXT("Will run up to %d pull batches concurrently"), ConcurrentTasks);
+
+	FWorkQueue WorkQueue(PayloadIds, BatchSize);
+
+	UE::Tasks::FTaskEvent Event(UE_SOURCE_LOCATION);
+
+	std::atomic<int32> NumTasks = 0;
+	double LogTimer = FPlatformTime::Seconds();
+
+	bool bHasErrors = false;
+
+	while (NumTasks != 0 || !WorkQueue.IsEmpty())
+	{
+		int32 NumTaskAllowed = ConcurrentTasks - NumTasks;
+		while (NumTaskAllowed > 0 && !WorkQueue.IsEmpty())
+		{
+			NumTasks++;
+
+			UE::Tasks::Launch(UE_SOURCE_LOCATION,
+				[Job = WorkQueue.GetJob(), &System, &NumCompletedPayloads, &NumTasks, &Event, &bHasErrors , &Callback]()
+				{
+					TArray<FPullRequest> Requests = ToRequestArray(Job);
+
+					if (!System.PullData(Requests))
+					{
+						bHasErrors = true;
+					}
+					
+					for (const FPullRequest& Request : Requests)
+					{
+						Callback(Request);
+					}
+					
+					NumCompletedPayloads += Requests.Num();
+					--NumTasks;
+
+					Event.Trigger();
+				});
+
+			--NumTaskAllowed;
+		}
+
+		Event.Wait(FTimespan::FromSeconds(30.0));
+
+		if (FPlatformTime::Seconds() - LogTimer >= 30.0)
+		{
+			const int32 CurrentCompletedPayloads = NumCompletedPayloads;
+			const float Progress = ((float)CurrentCompletedPayloads / (float)NumPayloads) * 100.0f;
+
+			UE_LOG(LogVirtualization, Display, TEXT("Cached %d/%d (%.1f%%)"), 
+				ProgressString != nullptr? ProgressString : TEXT("Processed"),
+				CurrentCompletedPayloads, NumPayloads, Progress);
+
+			LogTimer = FPlatformTime::Seconds();
+		}
+	}
+
+	return !bHasErrors;
+}
+
 TArray<FString> FindPackages(EFindPackageFlags Flags)
 {
 	TRACE_CPUPROFILER_EVENT_SCOPE(FindPackages);
@@ -169,7 +248,7 @@ TArray<FPullRequest> ToRequestArray(TConstArrayView<FIoHash> IdentifierArray)
 	return Requests;
 }
 
-FWorkQueue::FWorkQueue(const TArray<FIoHash>& InWork, int32 JobSize)
+FWorkQueue::FWorkQueue(const TConstArrayView<FIoHash> InWork, int32 JobSize)
 	: Work(InWork)
 {
 	CreateJobs(JobSize);
diff --git a/Engine/Source/Editor/VirtualizationEditor/Private/CommandletUtils.h b/Engine/Source/Editor/VirtualizationEditor/Private/CommandletUtils.h
--- a/Engine/Source/Editor/VirtualizationEditor/Private/CommandletUtils.h
+++ b/Engine/Source/Editor/VirtualizationEditor/Private/CommandletUtils.h
@@ -4,6 +4,7 @@
 
 #include "Containers/Array.h"
 #include "Containers/UnrealString.h"
+#include "Templates/Function.h"
 
 namespace UE { class FPackageTrailer; }
 namespace UE::Virtualization { struct FPullRequest; }
@@ -12,6 +13,21 @@ struct FIoHash;
 namespace UE::Virtualization // Utility functions
 {
 
+/*
+ * Utility to pull payloads in batches from many threads at once which can be much faster than pulling from a single thread
+ * in a single batch. The provided callback will be invoked each time that a request is completed.
+ * Every 30 seconds an update will be printed to the log showing how many payloads have been processed so that the user does
+ * not think that the process has stalled.
+ * 
+ * @param PayloadIds		The payloads that should be pulled
+ * @param BatchSize			The max number of payload ids that should be in each batch
+ * @param ProgressString	An optional string to be printed as part of the status shown every 30 seconds, can be nullptr.
+ * @param Callback			Called on every request once it has been processed.
+ * 
+ * @ return True if all requests succeeded, false if there were failures.
+ */
+bool PullPayloadsThreaded(TConstArrayView<FIoHash> PayloadIds, int32 BatchSize, const TCHAR* ProgressString, TUniqueFunction<void(const UE::Virtualization::FPullRequest& Response)>&& Callback);
+
 /** Used to customize package discovery behavior */
 enum EFindPackageFlags
 {
@@ -77,7 +93,7 @@ class FWorkQueue
 public:
 	using FJob = TArrayView<const FIoHash>;
 
-	FWorkQueue(const TArray<FIoHash>& InWork, int32 JobSize);
+	FWorkQueue(TConstArrayView<FIoHash> InWork, int32 JobSize);
 	FWorkQueue(TArray<FIoHash>&& InWork, int32 JobSize);
 	~FWorkQueue() = default;
 
diff --git a/Engine/Source/Editor/VirtualizationEditor/Private/PrecachePayloadsCommandlet.cpp b/Engine/Source/Editor/VirtualizationEditor/Private/PrecachePayloadsCommandlet.cpp
--- a/Engine/Source/Editor/VirtualizationEditor/Private/PrecachePayloadsCommandlet.cpp
+++ b/Engine/Source/Editor/VirtualizationEditor/Private/PrecachePayloadsCommandlet.cpp
@@ -2,9 +2,7 @@
 
 #include "PrecachePayloadsCommandlet.h"
 
-#include "Async/ParallelFor.h"
 #include "CommandletUtils.h"
-#include "Tasks/Task.h"
 #include "UObject/PackageTrailer.h"
 #include "Virtualization/VirtualizationSystem.h"
 
@@ -31,77 +29,15 @@ int32 UPrecachePayloadsCommandlet::Main(const FString& Params)
 
 	UE_LOG(LogVirtualization, Display, TEXT("Found %d virtualized payloads to precache"), PayloadIds.Num());
 	UE_LOG(LogVirtualization, Display, TEXT("Precaching payloads..."));
-
-	UE::Virtualization::IVirtualizationSystem& System = UE::Virtualization::IVirtualizationSystem::Get();
-	{
-		TRACE_CPUPROFILER_EVENT_SCOPE(Precache_ThreadedBatches);
-
-		std::atomic<int32> NumCompletedPayloads = 0;
-
-		const int32 NumPayloads = PayloadIds.Num();
-		const int32 BatchSize = 64;
-
-		// This seems to be the sweet spot when it comes to our internal infrastructure, so use that as the default.
-		const int32 MaxConcurrentTasks = 16;
-
-		// We always want to leave at least one foreground worker free to avoid saturation. If we issue too many
-		// concurrent task then we can potentially cause the DDC/Zen to be unable to run clean up tasks for long 
-		// periods of times which can cause quite high memory spikes.
-		const int32 ConcurrentTasks = FMath::Min(MaxConcurrentTasks, FTaskGraphInterface::Get().GetNumWorkerThreads() - 1);
-
-		UE_LOG(LogVirtualization, Display, TEXT("Will run up to %d precache tasks concurrently"), ConcurrentTasks);
-
-		UE::Virtualization::FWorkQueue WorkQueue(MoveTemp(PayloadIds), BatchSize);
-
-		UE::Tasks::FTaskEvent Event(UE_SOURCE_LOCATION);
-
-		std::atomic<int32> NumTasks = 0;
-		double LogTimer = FPlatformTime::Seconds();
-
-		while (NumTasks != 0 || !WorkQueue.IsEmpty())
+	
+	const int32 BatchSize = 64;
+	UE::Virtualization::PullPayloadsThreaded(PayloadIds, BatchSize, TEXT("Cached"), [](const UE::Virtualization::FPullRequest& Request)
 		{
-			int32 NumTaskAllowed = ConcurrentTasks - NumTasks;
-			while (NumTaskAllowed > 0 && !WorkQueue.IsEmpty())
+			if (!Request.IsSuccess())
 			{
-				NumTasks++;
-
-				UE::Tasks::Launch(UE_SOURCE_LOCATION,
-					[Job = WorkQueue.GetJob(), &System, &NumCompletedPayloads, &NumTasks, &Event]()
-					{
-						TArray<UE::Virtualization::FPullRequest> Requests = UE::Virtualization::ToRequestArray(Job);
-
-						if (!System.PullData(Requests))
-						{
-							for (const UE::Virtualization::FPullRequest& Request : Requests)
-							{
-								if (!Request.IsSuccess())
-								{
-									UE_LOG(LogVirtualization, Error, TEXT("%s: Failed to precache payload"), *LexToString(Request.GetIdentifier()));
-								}
-							}
-						}
-
-						NumCompletedPayloads += Requests.Num();
-						--NumTasks;
-
-						Event.Trigger();
-					});
-
-				--NumTaskAllowed;
+				UE_LOG(LogVirtualization, Error, TEXT("%s: Failed to precache payload"), *LexToString(Request.GetIdentifier()));
 			}
-
-			Event.Wait(FTimespan::FromSeconds(30.0));
-
-			if (FPlatformTime::Seconds() - LogTimer >= 30.0)
-			{
-				const int32 CurrentCompletedPayloads = NumCompletedPayloads;
-				const float Progress = ((float)CurrentCompletedPayloads / (float)NumPayloads) * 100.0f;
-				UE_LOG(LogVirtualization, Display, TEXT("Cached %d/%d (%.1f%%)"), CurrentCompletedPayloads, NumPayloads, Progress);
-
-				LogTimer = FPlatformTime::Seconds();
-			}
-		}
-	}
+		});
 
 	UE_LOG(LogVirtualization, Display, TEXT("Precaching complete!"));
 
