Add a new cmdline option '-ValidateContent' to the 'ValidateVirtualizedContent commandlet

#rnx

- The new cmdline option will download each payload and make sure that the data is correct in addition to checking that the payload exists.
- Note: There is no way to control which backend returns the payload so the VA graph, or cmdline parameters would need to be changed be the caller to make sure they are checking the backend that they are interested in.
-- This new option is fairly slow and is intended to be used as part of an automated checking system rather than by users themselves, so needing to change the graph and running the commandlet once per backend that needs testing shouldn't be too hard to do.

diff --git a/Engine/Source/Editor/VirtualizationEditor/Private/ValidateVirtualizedContentCommandlet.cpp b/Engine/Source/Editor/VirtualizationEditor/Private/ValidateVirtualizedContentCommandlet.cpp
--- a/Engine/Source/Editor/VirtualizationEditor/Private/ValidateVirtualizedContentCommandlet.cpp
+++ b/Engine/Source/Editor/VirtualizationEditor/Private/ValidateVirtualizedContentCommandlet.cpp
@@ -6,17 +6,69 @@
 #include "UObject/PackageTrailer.h"
 #include "Virtualization/VirtualizationSystem.h"
 
+namespace UE
+{
+
+/** Utility class so we can find all of the packages that use a given payload id */
+class FIoHashToPackagesLookup
+{
+public:
+	FIoHashToPackagesLookup(const TMap<FString, UE::FPackageTrailer>& InPackages)
+	{
+		PackagePaths.SetNum(InPackages.Num());
+
+		for (const TPair<FString, UE::FPackageTrailer>& Package : InPackages)
+		{
+			const int32 Index = PackagePaths.Num();
+			PackagePaths.Add(Package.Key);
+
+			const TArray<FIoHash> VirtualizedPayloads = Package.Value.GetPayloads(UE::EPayloadStorageType::Virtualized);
+			for (const FIoHash& Id : VirtualizedPayloads)
+			{
+				PackageLookup.FindOrAdd(Id).Add(Index);
+			}
+		}
+	}
+
+	void PrintError(const FIoHash& Id, const TCHAR* ErrorMessage) const
+	{
+		TStringBuilder<1028> FinalMessage;
+		FinalMessage << TEXT("Payload ") << Id << TEXT(" - ") << ErrorMessage;
+
+		for (int32 PackageIndex : PackageLookup[Id])
+		{
+			FinalMessage << LINE_TERMINATOR << TEXT(" \t") << PackagePaths[PackageIndex];
+		}
+
+		UE_LOG(LogVirtualization, Error, TEXT("%s"), FinalMessage.ToString());
+	}
+
+private:
+	TArray<FString> PackagePaths;
+	TMap<FIoHash, TArray<uint32, TInlineAllocator<4>>> PackageLookup;
+};
+
+} //namespace UE
+
 UValidateVirtualizedContentCommandlet::UValidateVirtualizedContentCommandlet(const FObjectInitializer& ObjectInitializer)
 	: Super(ObjectInitializer)
 {
 }
 
+
 int32 UValidateVirtualizedContentCommandlet::Main(const FString& Params)
 {
 	using namespace UE::Virtualization;
 
 	TRACE_CPUPROFILER_EVENT_SCOPE(UValidateVirtualizedContentCommandlet);
 
+	TArray<FString> Tokens;
+	TArray<FString> Switches;
+
+	ParseCommandLine(*Params, Tokens, Switches);
+
+	const bool bValidateContent = Switches.Contains(TEXT("ValidateContent"));
+
 	UE_LOG(LogVirtualization, Display, TEXT("Finding packages in the project..."));
 	TArray<FString> PackagePaths = FindPackages(EFindPackageFlags::ExcludeEngineContent);
 	UE_LOG(LogVirtualization, Display, TEXT("Found %d package(s)"), PackagePaths.Num());
@@ -28,6 +80,36 @@ int32 UValidateVirtualizedContentCommandlet::Main(const FString& Params)
 	FindVirtualizedPayloadsAndTrailers(PackagePaths, Packages, Payloads);
 	UE_LOG(LogVirtualization, Display, TEXT("Found %d virtualized package(s) with %d unique payload(s)"), Packages.Num(), Payloads.Num());
 
+	UE::FIoHashToPackagesLookup PkgLookupTable(Packages);
+
+	int32 ErrorCount = 0;
+	if (bValidateContent)
+	{
+		ErrorCount = ValidatePayloadContent(Packages, Payloads, PkgLookupTable);
+	}
+	else
+	{
+		ErrorCount = ValidatePayloadsExists(Packages, Payloads);
+	}
+
+	if (ErrorCount == 0)
+	{
+		UE_LOG(LogVirtualization, Display, TEXT("All virtualized payloads could be found in persistent storage"));
+		return 0;
+	}
+	else
+	{
+		UE_LOG(LogVirtualization, Error, TEXT("%d/%d package(s) had at least one virtualized payload missing from persistent storage"), ErrorCount, Packages.Num());
+		return 0;
+	}
+}
+
+int32 UValidateVirtualizedContentCommandlet::ValidatePayloadsExists(const TMap<FString, UE::FPackageTrailer>& Packages, const TSet<FIoHash>& Payloads)
+{
+	using namespace UE::Virtualization;
+
+	UE_LOG(LogVirtualization, Display, TEXT("Validating payloads existence.."));
+
 	IVirtualizationSystem& System = IVirtualizationSystem::Get();
 
 	UE_LOG(LogVirtualization, Display, TEXT("Querying the state of the virtualized payload(s) in persistent storage..."));
@@ -49,7 +131,8 @@ int32 UValidateVirtualizedContentCommandlet::Main(const FString& Params)
 		{
 			bool bFoundErrors = false;
 
-			const TArray<FIoHash> VirtualizedPayloads = Package.Value.GetPayloads(UE::EPayloadStorageType::Virtualized);
+			const UE::FPackageTrailer& Trailer = Package.Value;
+			const TArray<FIoHash> VirtualizedPayloads = Trailer.GetPayloads(UE::EPayloadStorageType::Virtualized);
 
 			for (const FIoHash& PayloadId : VirtualizedPayloads)
 			{
@@ -74,14 +157,52 @@ int32 UValidateVirtualizedContentCommandlet::Main(const FString& Params)
 		}
 	}
 
-	if (ErrorCount == 0)
-	{
-		UE_LOG(LogVirtualization, Display, TEXT("All virtualized payloads could be found in persistent storage"));
-		return 0;
-	}
-	else
-	{
-		UE_LOG(LogVirtualization, Error, TEXT("%d/%d package(s) had at least one virtualized payload missing from persistent storage"), ErrorCount, Packages.Num());
-		return 0;
-	}
+	return ErrorCount;
 }
+
+int32 UValidateVirtualizedContentCommandlet::ValidatePayloadContent(const TMap<FString, UE::FPackageTrailer>& Packages, const TSet<FIoHash>& Payloads, const UE::FIoHashToPackagesLookup& PkgLookupTable)
+{
+	UE_LOG(LogVirtualization, Display, TEXT("Validating payloads existence and content..."));
+
+	std::atomic<int32> ErrorCount = 0;
+
+	const int32 BatchSize = 64;
+	UE::Virtualization::PullPayloadsThreaded(Payloads.Array(), BatchSize, TEXT("Validated"), [&PkgLookupTable , &ErrorCount](const UE::Virtualization::FPullRequest& Request)
+		{
+			if (!Request.IsSuccess() || Request.GetPayload().IsNull())
+			{
+				PkgLookupTable.PrintError(Request.GetIdentifier(), TEXT("could not be pulled"));
+
+				ErrorCount++;
+				return;
+			}
+			
+			if (Request.GetPayload().GetRawHash() != Request.GetIdentifier())
+			{
+				PkgLookupTable.PrintError(Request.GetIdentifier(), TEXT("pulled a payload with the wrong hash"));
+
+				ErrorCount++;
+				return;
+			}
+
+			FSharedBuffer UncompressedPayload = Request.GetPayload().Decompress();
+			if (UncompressedPayload.IsNull())
+			{
+				PkgLookupTable.PrintError(Request.GetIdentifier(), TEXT("could not be decompressed"));
+
+				ErrorCount++;
+				return;
+			}
+
+			if (FIoHash::HashBuffer(UncompressedPayload) != Request.GetIdentifier())
+			{
+				PkgLookupTable.PrintError(Request.GetIdentifier(), TEXT("is corrupted"));
+
+				ErrorCount++;
+				return;
+			}
+		});
+
+	return ErrorCount;
+}
+
diff --git a/Engine/Source/Editor/VirtualizationEditor/Private/ValidateVirtualizedContentCommandlet.h b/Engine/Source/Editor/VirtualizationEditor/Private/ValidateVirtualizedContentCommandlet.h
--- a/Engine/Source/Editor/VirtualizationEditor/Private/ValidateVirtualizedContentCommandlet.h
+++ b/Engine/Source/Editor/VirtualizationEditor/Private/ValidateVirtualizedContentCommandlet.h
@@ -7,6 +7,12 @@
 
 #include "ValidateVirtualizedContentCommandlet.generated.h"
 
+namespace UE { class FPackageTrailer; }
+namespace UE { class FIoHashToPackagesLookup; }
+
+class FIoHashToPackagesLookup;
+struct FIoHash;
+
 /**
  * Iterates over all of the packages in a project and identifies which packages contain
  * references to virtualized payloads. The commandlet will then check that all virtualized
@@ -22,6 +28,7 @@ UCLASS()
 class UValidateVirtualizedContentCommandlet
 	: public UCommandlet
 {
+private:
 	GENERATED_UCLASS_BODY()
 
 	//~ Begin UCommandlet Interface
@@ -29,4 +36,8 @@ class UValidateVirtualizedContentCommandlet
 	//~ End UCommandlet Interface
 
 	static int32 StaticMain(const FString& Params);
+
+	int32 ValidatePayloadsExists(const TMap<FString, UE::FPackageTrailer>& Packages, const TSet<FIoHash>& Payloads);
+
+	int32 ValidatePayloadContent(const TMap<FString, UE::FPackageTrailer>& Packages, const TSet<FIoHash>& Payloads, const UE::FIoHashToPackagesLookup& PkgLookupTable);
 };
