Allow the virtualize and rehydrate project commandlets to work on specific directories if requested.

- The commandlets now accept '-PackageDir=' and '-PackageFolder=' commnad line args.

diff --git a/Engine/Source/Editor/VirtualizationEditor/Private/RehydrateProjectCommandlet.cpp b/Engine/Source/Editor/VirtualizationEditor/Private/RehydrateProjectCommandlet.cpp
--- a/Engine/Source/Editor/VirtualizationEditor/Private/RehydrateProjectCommandlet.cpp
+++ b/Engine/Source/Editor/VirtualizationEditor/Private/RehydrateProjectCommandlet.cpp
@@ -27,7 +27,7 @@ int32 URehydrateProjectCommandlet::Main(const FString& Params)
 	}
 
 	UE_LOG(LogVirtualization, Display, TEXT("Scanning for project packages..."));
-	TArray<FString> Packages = UE::Virtualization::FindPackages(UE::Virtualization::EFindPackageFlags::ExcludeEngineContent);
+	TArray<FString> Packages = UE::Virtualization::DiscoverPackages(Params, UE::Virtualization::EFindPackageFlags::ExcludeEngineContent);
 
 	UE_LOG(LogVirtualization, Display, TEXT("Rehydrating packages..."));
 
diff --git a/Engine/Source/Editor/VirtualizationEditor/Private/VirtualizeProjectCommandlet.cpp b/Engine/Source/Editor/VirtualizationEditor/Private/VirtualizeProjectCommandlet.cpp
--- a/Engine/Source/Editor/VirtualizationEditor/Private/VirtualizeProjectCommandlet.cpp
+++ b/Engine/Source/Editor/VirtualizationEditor/Private/VirtualizeProjectCommandlet.cpp
@@ -27,8 +27,7 @@ int32 UVirtualizeProjectCommandlet::Main(const FString& Params)
 	}
 
 	UE_LOG(LogVirtualization, Display, TEXT("Scanning for project packages..."));
-	TArray<FString> Packages = UE::Virtualization::FindPackages(UE::Virtualization::EFindPackageFlags::ExcludeEngineContent);
-
+	TArray<FString> Packages = UE::Virtualization::DiscoverPackages(Params, UE::Virtualization::EFindPackageFlags::ExcludeEngineContent);
 	UE_LOG(LogVirtualization, Display, TEXT("Virtualizing packages..."));
 
 	UE::Virtualization::IVirtualizationSystem& System = UE::Virtualization::IVirtualizationSystem::Get();
