Add a new cmdline parameter '-LocalOnly' to the 'GeneratePayloadManifest commandlet. When used only payloads that are still stored locally will be added to the manifest.

diff --git a/Engine/Source/Editor/VirtualizationEditor/Private/GeneratePayloadManifestCommandlet.cpp b/Engine/Source/Editor/VirtualizationEditor/Private/GeneratePayloadManifestCommandlet.cpp
--- a/Engine/Source/Editor/VirtualizationEditor/Private/GeneratePayloadManifestCommandlet.cpp
+++ b/Engine/Source/Editor/VirtualizationEditor/Private/GeneratePayloadManifestCommandlet.cpp
@@ -18,6 +18,12 @@ int32 UGeneratePayloadManifestCommandlet::Main(const FString& Params)
 {
 	TRACE_CPUPROFILER_EVENT_SCOPE(UGeneratePayloadManifestCommandlet);
 
+	if (!ParseCmdline(Params))
+	{
+		UE_LOG(LogVirtualization, Error, TEXT("Failed to parse the command line correctly"));
+		return -1;
+	}
+
 	TArray<FString> PackageNames = UE::Virtualization::DiscoverPackages(Params, UE::Virtualization::EFindPackageFlags::ExcludeEngineContent);
 
 	UE_LOG(LogVirtualization, Display, TEXT("Found %d files to look in"), PackageNames.Num());
@@ -50,8 +56,13 @@ int32 UGeneratePayloadManifestCommandlet::Main(const FString& Params)
 				{
 					PackageTrailerCount++;
 
-					Trailer.ForEachPayload([&ManifestFile,&PackagePath](const FIoHash& Id, uint64 SizeOnDisk, uint64 RawSize, UE::EPayloadAccessMode Mode, UE::Virtualization::EPayloadFilterReason Filter)->void
+					Trailer.ForEachPayload([&ManifestFile, &PackagePath, bLocalOnly = bLocalOnly](const FIoHash& Id, uint64 SizeOnDisk, uint64 RawSize, UE::EPayloadAccessMode Mode, UE::Virtualization::EPayloadFilterReason Filter)->void
 						{
+							if (bLocalOnly && Mode != UE::EPayloadAccessMode::Local)
+							{
+								return;
+							}
+
 							TAnsiStringBuilder<256> LineBuilder;
 							LineBuilder << PackagePath << "," << Id << "," << SizeOnDisk << "," << RawSize << "," << Mode << "," << *LexToString(Filter) <<"\n";
 							ManifestFile->Serialize((void*)LineBuilder.ToString(), LineBuilder.Len());
@@ -74,4 +85,16 @@ int32 UGeneratePayloadManifestCommandlet::Main(const FString& Params)
 	UE_LOG(LogVirtualization, Display, TEXT("Manifest written to '%s'"), *CSVPath);
 
 	return  0;
-}
\ No newline at end of file
+}
+
+bool UGeneratePayloadManifestCommandlet::ParseCmdline(const FString& Params)
+{
+	TArray<FString> Tokens;
+	TArray<FString> Switches;
+
+	ParseCommandLine(*Params, Tokens, Switches);
+
+	bLocalOnly = Switches.Contains(TEXT("LocalOnly"));
+
+	return true;
+}
diff --git a/Engine/Source/Editor/VirtualizationEditor/Private/GeneratePayloadManifestCommandlet.h b/Engine/Source/Editor/VirtualizationEditor/Private/GeneratePayloadManifestCommandlet.h
--- a/Engine/Source/Editor/VirtualizationEditor/Private/GeneratePayloadManifestCommandlet.h
+++ b/Engine/Source/Editor/VirtualizationEditor/Private/GeneratePayloadManifestCommandlet.h
@@ -28,4 +28,10 @@ class UGeneratePayloadManifestCommandlet
 	//~ End UCommandlet Interface
 
 	static int32 StaticMain(const FString& Params);
+
+	bool ParseCmdline(const FString& Params);
+
+private:
+
+	bool bLocalOnly = false;
 };
