Allow the output of the GeneratePayloadManifest commandlet to be spit up over multiple csv files if it gets too large.

- If a project contained too many packages the final .csv output file might be too large to import into some spread sheet systems, in these cases we need to split the output into multiple smaller .csv files.
- The output is now sent to its own subdirectory under saved, ''/payloadmanifest/"
- When done the commandlet will report the file path of each .csv written to disk
- Add the number of payloads discovered to the final output summary.

diff --git a/Engine/Source/Editor/VirtualizationEditor/Private/GeneratePayloadManifestCommandlet.cpp b/Engine/Source/Editor/VirtualizationEditor/Private/GeneratePayloadManifestCommandlet.cpp
--- a/Engine/Source/Editor/VirtualizationEditor/Private/GeneratePayloadManifestCommandlet.cpp
+++ b/Engine/Source/Editor/VirtualizationEditor/Private/GeneratePayloadManifestCommandlet.cpp
@@ -9,6 +9,83 @@
 #include "Misc/Paths.h"
 #include "UObject/PackageTrailer.h"
 
+class FSpreadSheet
+{
+public:
+	bool OpenNewSheet()
+	{
+		Ar.Reset();
+
+		const TStringBuilder<512> CSVPath = GetFilePath(++NumSheets);
+
+		Ar = TUniquePtr<FArchive>(IFileManager::Get().CreateFileWriter(CSVPath.ToString()));
+		if (Ar.IsValid())
+		{
+			const ANSICHAR* Headings = "Path,PayloadId,SizeOnDisk,UncompressedSize,StorageType, FilterReason\n";
+			Ar->Serialize((void*)Headings, FPlatformString::Strlen(Headings));
+
+			return true;
+		}
+		else
+		{
+			UE_LOG(LogVirtualization, Error, TEXT("Failed to open '%s' for write"), CSVPath.ToString());
+			return false;
+		}
+	}
+
+	bool PrintRow(FAnsiStringView Row)
+	{
+		if (++NumRowsInSheet >= MaxNumRowsPerSheet || (Ar->Tell() + Row.Len()) >= MaxFileSize)
+		{
+			if (!OpenNewSheet())
+			{
+				return false;
+			}
+
+			NumRowsInSheet = 0;
+		}
+
+		Ar->Serialize((void*)Row.GetData(), Row.Len() * sizeof(FAnsiStringView::ElementType));
+
+		return true;
+	}
+
+	void Flush()
+	{
+		Ar.Reset();
+	}
+
+	FString GetDebugInfo() const
+	{
+		TStringBuilder<1024> Output;
+
+		Output << TEXT("Manifest is comprised of ") << NumSheets << TEXT(" sheet(s) written to:");
+
+		for (int32 Index = 0; Index < NumSheets; ++Index)
+		{
+			Output << LINE_TERMINATOR << GetFilePath(Index);
+		}
+
+		return Output.ToString();
+	}
+
+private:
+
+	static TStringBuilder<512> GetFilePath(int32 Index)
+	{
+		return WriteToString<512>(FPaths::ProjectSavedDir(), TEXT("PayloadManifest/sheet"), Index, TEXT(".csv"));
+	}
+
+	// With any values above these we had trouble importing the finished cvs into various spread sheet programs
+	const int64 MaxNumRowsPerSheet = 250000;
+	const int64 MaxFileSize = 45 * 1024 * 1024;
+
+	int32 NumSheets = 0;
+	int32 NumRowsInSheet = 0;
+
+	TUniquePtr<FArchive> Ar;
+};
+
 UGeneratePayloadManifestCommandlet::UGeneratePayloadManifestCommandlet(const FObjectInitializer& ObjectInitializer)
 	: Super(ObjectInitializer)
 {
@@ -29,18 +106,14 @@ int32 UGeneratePayloadManifestCommandlet::Main(const FString& Params)
 	UE_LOG(LogVirtualization, Display, TEXT("Found %d files to look in"), PackageNames.Num());
 
 	int32 PackageTrailerCount = 0;
+	int32 PayloadCount = 0;
 
-	const FString CSVPath = FPaths::ProjectSavedDir() / TEXT("payload_manifest.csv");
-	TUniquePtr<FArchive> ManifestFile(IFileManager::Get().CreateFileWriter(*CSVPath));
-	if (!ManifestFile.IsValid())
+	FSpreadSheet Sheet;
+	if (!Sheet.OpenNewSheet())
 	{
-		UE_LOG(LogVirtualization, Error, TEXT("Failed to open '%s' for write"), *CSVPath);
 		return 1;
 	}
 
-	const ANSICHAR* Headings = "Path,PayloadId,SizeOnDisk,UncompressedSize,StorageType, FilterReason\n";
-	ManifestFile->Serialize((void*)Headings, FPlatformString::Strlen(Headings));
-
 	{
 		TRACE_CPUPROFILER_EVENT_SCOPE(ParsePackageTrailers);
 
@@ -56,16 +129,19 @@ int32 UGeneratePayloadManifestCommandlet::Main(const FString& Params)
 				{
 					PackageTrailerCount++;
 
-					Trailer.ForEachPayload([&ManifestFile, &PackagePath, bLocalOnly = bLocalOnly](const FIoHash& Id, uint64 SizeOnDisk, uint64 RawSize, UE::EPayloadAccessMode Mode, UE::Virtualization::EPayloadFilterReason Filter)->void
+					Trailer.ForEachPayload([&Sheet, &PackagePath, &PayloadCount, bLocalOnly = bLocalOnly](const FIoHash& Id, uint64 SizeOnDisk, uint64 RawSize, UE::EPayloadAccessMode Mode, UE::Virtualization::EPayloadFilterReason Filter)->void
 						{
 							if (bLocalOnly && Mode != UE::EPayloadAccessMode::Local)
 							{
 								return;
 							}
 
+							PayloadCount++;
+
 							TAnsiStringBuilder<256> LineBuilder;
 							LineBuilder << PackagePath << "," << Id << "," << SizeOnDisk << "," << RawSize << "," << Mode << "," << *LexToString(Filter) <<"\n";
-							ManifestFile->Serialize((void*)LineBuilder.ToString(), LineBuilder.Len());
+							
+							Sheet.PrintRow(LineBuilder);
 						});
 				}
 			}
@@ -79,10 +155,10 @@ int32 UGeneratePayloadManifestCommandlet::Main(const FString& Params)
 		}
 	}
 
-	ManifestFile.Reset(); // Flush file to disk
+	Sheet.Flush();
 
-	UE_LOG(LogVirtualization, Display, TEXT("Found %d package trailers"), PackageTrailerCount);
-	UE_LOG(LogVirtualization, Display, TEXT("Manifest written to '%s'"), *CSVPath);
+	UE_LOG(LogVirtualization, Display, TEXT("Found %d package trailers with %d payloads"), PackageTrailerCount, PayloadCount);
+	UE_LOG(LogVirtualization, Display, TEXT("%s"), *Sheet.GetDebugInfo());
 
 	return  0;
 }
