Fix some typos in code documentation

#rnx

diff --git a/Engine/Source/Editor/VirtualizationEditor/Private/GeneratePayloadManifestCommandlet.h b/Engine/Source/Editor/VirtualizationEditor/Private/GeneratePayloadManifestCommandlet.h
--- a/Engine/Source/Editor/VirtualizationEditor/Private/GeneratePayloadManifestCommandlet.h
+++ b/Engine/Source/Editor/VirtualizationEditor/Private/GeneratePayloadManifestCommandlet.h
@@ -10,10 +10,10 @@
 /**
  * Creates a csv file containing info about all of the payloads in a set of packages.
  * By default the commandlet will parse the payloads of all packages in the current
- * project but this can be overriden with the cmdline switch -PackageDir=XYZ which
+ * project but this can be overridden with the cmdline switch -PackageDir=XYZ which
  * will allow the commandlet to parse the payloads of the packages in a given directory.
  * 
- * Because the commmandlet is the VirtualizationEditor module it needs to be invoked 
+ * Because the commandlet is the VirtualizationEditor module it needs to be invoked 
  * with the command line:
  * -run="VirtualizationEditor.GeneratePayloadManifest"
  */
