Add a new console command 'ValidatePackagePayloads' as a quick way to check the status of payloads in a package.

#rnx

- Each payload in the package is loaded or pulled and checked to make sure that it:
-- Could be accessed.
-- The hash stored in the FCompressedBuffer is the same as the original requested hash.
-- That the payload can be decompressed.
-- That the decompressed payload when hashed has the same hash as the original request.
- This should catch cases where the payload cannot be found or is corrupted in some way.

diff --git a/Engine/Source/Editor/VirtualizationEditor/Private/ConsoleCommandUtils.cpp b/Engine/Source/Editor/VirtualizationEditor/Private/ConsoleCommandUtils.cpp
new file mode 100644
--- /dev/null
+++ b/Engine/Source/Editor/VirtualizationEditor/Private/ConsoleCommandUtils.cpp
@@ -0,0 +1,114 @@
+// Copyright Epic Games, Inc. All Rights Reserved.
+
+#include "ConsoleCommandUtils.h"
+
+#include "HAL/FileManager.h"
+#include "Misc/PackagePath.h"
+#include "Serialization/Archive.h"
+#include "UObject/PackageFileSummary.h"
+#include "UObject/PackageResourceManager.h"
+
+namespace UE::Virtualization
+{
+
+void LogPackageError(FArchive* Ar, const TCHAR* DebugName)
+{
+	if (Ar == nullptr)
+	{
+		UE_LOG(LogVirtualization, Error, TEXT("Could not find the package file: '%s'"), *DebugName);
+		return;
+	}
+
+	FPackageFileSummary Summary;
+	*Ar << Summary;
+
+	if (Ar->IsError())
+	{
+		UE_LOG(LogVirtualization, Error, TEXT("Could not find load the package summary from disk: '%s'"), DebugName);
+		return;
+	}
+
+	if (Summary.Tag != PACKAGE_FILE_TAG)
+	{
+		UE_LOG(LogVirtualization, Error, TEXT("Package summary seems to be corrupted: '%s'"), DebugName);
+		return;
+	}
+
+	int32 PackageVersion = Summary.GetFileVersionUE().ToValue();
+	if (PackageVersion >= (int32)EUnrealEngineObjectUE5Version::PAYLOAD_TOC)
+	{
+		UE_LOG(LogVirtualization, Error, TEXT("Package trailer is missing from the package file: '%s'"), DebugName);
+		return;
+	}
+
+	UE_LOG(LogVirtualization, Error, TEXT("Package is tool old (version %d) to have a package trailer (version %d): '%s'"), PackageVersion, EUnrealEngineObjectUE5Version::PAYLOAD_TOC, DebugName);
+}
+
+void LogPackageError(const FPackagePath& Path)
+{
+	const FString DebugName = Path.GetPackageName();
+	TUniquePtr<FArchive> PackageAr = IPackageResourceManager::Get().OpenReadExternalResource(EPackageExternalResource::WorkspaceDomainFile, DebugName);
+
+	LogPackageError(PackageAr.Get(), *DebugName);
+}
+
+void LogPackageError(const FString& Path)
+{
+	TUniquePtr<FArchive> PackageAr(IFileManager::Get().CreateFileReader(*Path));
+
+	LogPackageError(PackageAr.Get(), *Path);
+}
+
+TArray<TPair<FString, UE::FPackageTrailer>> LoadPackageTrailerFromArgs(const TArray<FString>& Args)
+{
+	TArray<TPair<FString, UE::FPackageTrailer>> Packages;
+	Packages.Reserve(Args.Num());
+
+	for (const FString& Arg : Args)
+	{
+		FString PathString;
+
+		if (FPackageName::ParseExportTextPath(Arg, nullptr /*OutClassName*/, &PathString))
+		{
+			PathString = FPackageName::ObjectPathToPackageName(PathString);
+		}
+		else
+		{
+			PathString = Arg;
+		}
+
+		FPackageTrailer Trailer;
+
+		FPackagePath Path;
+		if (FPackagePath::TryFromMountedName(PathString, Path))
+		{
+			if (!FPackageTrailer::TryLoadFromPackage(Path, Trailer))
+			{
+				LogPackageError(Path);
+				continue;
+			}
+		}
+		else if (IFileManager::Get().FileExists(*PathString))
+		{
+			// IF we couldn't turn it into a FPackagePath it could be a path to a package not under any current mount point.
+			// So for a final attempt we will see if we can find the file on disk and load the package trailer that way.
+
+			if (!FPackageTrailer::TryLoadFromFile(PathString, Trailer))
+			{
+				LogPackageError(PathString);
+				continue;
+			}
+		}
+		else
+		{
+			UE_LOG(LogVirtualization, Error, TEXT("Arg '%s' could not be converted to a valid package path"), *Arg);
+			continue;
+		}
+
+		Packages.Add({ MoveTemp(PathString) , MoveTemp(Trailer) });	
+	}
+
+	return Packages;
+}
+
+} // namespace UE::Virtualization
diff --git a/Engine/Source/Editor/VirtualizationEditor/Private/ConsoleCommandUtils.h b/Engine/Source/Editor/VirtualizationEditor/Private/ConsoleCommandUtils.h
new file mode 100644
--- /dev/null
+++ b/Engine/Source/Editor/VirtualizationEditor/Private/ConsoleCommandUtils.h
@@ -0,0 +1,13 @@
+// Copyright Epic Games, Inc. All Rights Reserved.
+
+#pragma once
+
+#include "Containers/Array.h"
+#include "UObject/PackageTrailer.h"
+
+namespace UE::Virtualization
+{
+
+TArray<TPair<FString, UE::FPackageTrailer>> LoadPackageTrailerFromArgs(const TArray<FString>& Args);
+
+} //namespace UE::Virtualization
diff --git a/Engine/Source/Editor/VirtualizationEditor/Private/ValidatePackagePayloadsConsoleCommand.cpp b/Engine/Source/Editor/VirtualizationEditor/Private/ValidatePackagePayloadsConsoleCommand.cpp
new file mode 100644
--- /dev/null
+++ b/Engine/Source/Editor/VirtualizationEditor/Private/ValidatePackagePayloadsConsoleCommand.cpp
@@ -0,0 +1,121 @@
+// Copyright Epic Games, Inc. All Rights Reserved.
+
+#include "ConsoleCommandUtils.h"
+
+#include "CommandletUtils.h"
+#include "CoreGlobals.h"
+#include "HAL/IConsoleManager.h"
+#include "UObject/PackageResourceManager.h"
+#include "Virtualization/VirtualizationSystem.h"
+
+namespace UE::Virtualization
+{
+
+bool ValidatePayload(FIoHash Id, FCompressedBuffer CompressedPayload)
+{
+	if (CompressedPayload.IsNull())
+	{
+		UE_LOG(LogVirtualization, Error, TEXT("\tCould not load payload '%s'"), *LexToString(Id));
+		return false;
+	}
+
+	if (CompressedPayload.GetRawHash() != Id)
+	{
+		UE_LOG(LogVirtualization, Error, TEXT("\tPayload '%s' had the wrong hash (%s) after loading"), *LexToString(Id), *LexToString(CompressedPayload.GetRawHash()));
+		return false;
+	}
+
+	FSharedBuffer UncompressedPayload = CompressedPayload.Decompress();
+
+	if (UncompressedPayload.IsNull())
+	{
+		UE_LOG(LogVirtualization, Error, TEXT("\tCould not decompress payload '%s'"), *LexToString(Id));
+		return false;
+	}
+
+	FIoHash UncompressedHash = FIoHash::HashBuffer(UncompressedPayload);
+	if (UncompressedHash != Id)
+	{
+		UE_LOG(LogVirtualization, Error, TEXT("\tPayload '%s' had the wrong hash (%s) after decompression"), *LexToString(Id), *LexToString(UncompressedHash));
+		return false;
+	}
+
+	return true;
+}
+void ValidatePackagePayloads(const TArray<FString>& Args)
+{
+	if (Args.IsEmpty())
+	{
+		UE_LOG(LogVirtualization, Error, TEXT("Command 'ValidatePackagePayloads' called without any arguments"));
+		return;
+	}
+
+	TArray<TPair<FString, UE::FPackageTrailer>> Packages = LoadPackageTrailerFromArgs(Args);
+
+	if (Packages.IsEmpty())
+	{
+		UE_LOG(LogVirtualization, Error, TEXT("Command 'ValidatePackagePayloads' failed to find any valid packages from the provided arguments"));
+		return;
+	}
+
+	for (const TPair<FString, UE::FPackageTrailer>& Package : Packages)
+	{
+		const FString& PackagePath = Package.Key;
+		const UE::FPackageTrailer& Trailer = Package.Value;
+
+		UE_LOG(LogVirtualization, Display, TEXT("Validating package '%s'"), *PackagePath);
+
+		int32 ValidLocalPayloads = 0;
+		TArray<FIoHash> LocalPayloads = Trailer.GetPayloads(EPayloadStorageType::Local);
+
+		if (!LocalPayloads.IsEmpty())
+		{
+			UE_LOG(LogVirtualization, Display, TEXT("Checking Local Payloads..."));
+
+			TUniquePtr<FArchive> Ar = IPackageResourceManager::Get().OpenReadExternalResource(EPackageExternalResource::WorkspaceDomainFile, Package.Key);
+			check(Ar.IsValid());
+
+			for (const FIoHash& Id : LocalPayloads)
+			{
+				FCompressedBuffer CompressedPayload = Trailer.LoadLocalPayload(Id, *Ar);
+
+				if (ValidatePayload(Id, CompressedPayload))
+				{
+					ValidLocalPayloads++;
+				}
+			}
+		}
+
+		int32 ValidVirtualizedPayloads = 0;
+		TArray<FIoHash> VirtualizedPayloads = Trailer.GetPayloads(EPayloadStorageType::Virtualized);
+
+		if (!VirtualizedPayloads.IsEmpty())
+		{
+			UE_LOG(LogVirtualization, Display, TEXT("Checking Virtualized Payloads..."));
+
+			TArray<FPullRequest> Requests = ToRequestArray(VirtualizedPayloads);
+			IVirtualizationSystem::Get().PullData(Requests);
+
+			check(Requests.Num() == VirtualizedPayloads.Num());
+
+			for(int32 Index = 0; Index < VirtualizedPayloads.Num(); ++Index)
+			{
+				const FPullRequest& Request = Requests[Index];
+
+				if (ValidatePayload(Request.GetIdentifier(), Request.GetPayload()))
+				{
+					ValidLocalPayloads++;
+				}
+			}
+		}
+
+		UE_CLOG(!LocalPayloads.IsEmpty(), LogVirtualization, Display, TEXT("%d/%d local payloads are valid"), ValidLocalPayloads, LocalPayloads.Num());
+		UE_CLOG(!VirtualizedPayloads.IsEmpty(), LogVirtualization, Display, TEXT("%d/%d virtualized payloads are valid"), ValidVirtualizedPayloads, VirtualizedPayloads.Num());
+	}
+}
+
+static FAutoConsoleCommand CCmdDumpPayloadToc = FAutoConsoleCommand(
+	TEXT("ValidatePackagePayloads"),
+	TEXT("Checks all payloads that a package references and makes sure that they are valid"),
+	FConsoleCommandWithArgsDelegate::CreateStatic(ValidatePackagePayloads));
+} //namespace UE::Virtualization
