The commandlet 'VirtualizationEditor.PrecachePayloads' can now be configured to only precache specific directories by using the cmdline arg '-PackageDir=' or '-PackageFolder='

#rnx

diff --git a/Engine/Source/Editor/VirtualizationEditor/Private/CommandletUtils.cpp b/Engine/Source/Editor/VirtualizationEditor/Private/CommandletUtils.cpp
--- a/Engine/Source/Editor/VirtualizationEditor/Private/CommandletUtils.cpp
+++ b/Engine/Source/Editor/VirtualizationEditor/Private/CommandletUtils.cpp
@@ -9,6 +9,7 @@
 #include "IO/IoHash.h"
 #include "Misc/PackageName.h"
 #include "Misc/Paths.h"
+#include "String/ParseTokens.h"
 #include "Tasks/Task.h"	
 #include "UObject/PackageTrailer.h"
 #include "Virtualization/VirtualizationSystem.h"
@@ -98,6 +99,8 @@ TArray<FString> FindPackages(EFindPackageFlags Flags)
 {
 	TRACE_CPUPROFILER_EVENT_SCOPE(FindPackages);
 
+	UE_LOG(LogVirtualization, Display, TEXT("Searching for packages under the current project..."));
+
 	TArray<FString> PackagePaths;
 
 	FAssetRegistryModule& AssetRegistry = FModuleManager::LoadModuleChecked<FAssetRegistryModule>(TEXT("AssetRegistry"));
@@ -148,6 +151,8 @@ TArray<FString> FindPackagesInDirectory(const FString& DirectoryToSearch)
 		PackageNames.Add(FPaths::CreateStandardFilename(BasePath));
 	}
 
+	UE_LOG(LogVirtualization, Display, TEXT("Searching for packages under '%s'"), *DirectoryToSearch);
+
 	return PackageNames;
 }
 
@@ -155,10 +160,19 @@ TArray<FString> DiscoverPackages(const FString& CmdlineParams, EFindPackageFlags
 {
 	TRACE_CPUPROFILER_EVENT_SCOPE(DiscoverPackages);
 
-	FString PackageDir;
-	if (FParse::Value(*CmdlineParams, TEXT("PackageDir="), PackageDir) || FParse::Value(*CmdlineParams, TEXT("PackageFolder="), PackageDir))
+	FString PackageDirectories;
+	if (FParse::Value(*CmdlineParams, TEXT("PackageDir="), PackageDirectories) || FParse::Value(*CmdlineParams, TEXT("PackageFolder="), PackageDirectories))
 	{
-		return FindPackagesInDirectory(PackageDir);
+		TArray<FString> Packages;
+
+		UE::String::ParseTokensMultiple(PackageDirectories, TConstArrayView<TCHAR>({ '+', ',' }),
+			[&Packages](FStringView PackageDirectory)
+			{
+				TArray<FString> FoundPackages = FindPackagesInDirectory(FString(PackageDirectory));
+				Packages.Append(FoundPackages);
+			}, UE::String::EParseTokensOptions::Trim | UE::String::EParseTokensOptions::SkipEmpty);
+		
+		return Packages;
 	}
 	else
 	{
