Improve the parsing of '-PackageDir=' and '-PackageFolder=; in virtualization commandlets.

#rnx

- Allow more than one path to be specified at a time. Additional directories can be chained with either ',' or '+'
- Log each root directory searched or that the project itself is being searched.

diff --git a/Engine/Source/Editor/VirtualizationEditor/Private/PrecachePayloadsCommandlet.cpp b/Engine/Source/Editor/VirtualizationEditor/Private/PrecachePayloadsCommandlet.cpp
--- a/Engine/Source/Editor/VirtualizationEditor/Private/PrecachePayloadsCommandlet.cpp
+++ b/Engine/Source/Editor/VirtualizationEditor/Private/PrecachePayloadsCommandlet.cpp
@@ -15,7 +15,7 @@ int32 UPrecachePayloadsCommandlet::Main(const FString& Params)
 {
 	TRACE_CPUPROFILER_EVENT_SCOPE(UPrecachePayloadsCommandlet);
 
-	TArray<FString> PackageNames = UE::Virtualization::FindPackages(UE::Virtualization::EFindPackageFlags::ExcludeEngineContent);
+	TArray<FString> PackageNames = UE::Virtualization::DiscoverPackages(Params, UE::Virtualization::EFindPackageFlags::ExcludeEngineContent);
 
 	UE_LOG(LogVirtualization, Display, TEXT("Found %d packages"), PackageNames.Num());
 
