Improve the virtualization statistics panel.

#rnx

- Added boxes around each set of categories (pull, push, cache) to make it easier to see the differences between them.
- Moved the "Push" statistics to be last in order as it is the least interest of the statistics, infact when an external submit tool is being used with the editor it won't display any useful info at all (we should look into removing that set of stats if we can detect such tool use)
- Changed the total time spent in request to an average time per request
- Refactored the code to reduce code duplication.

diff --git a/Engine/Source/Editor/VirtualizationEditor/Private/SVirtualAssetsStatistics.cpp b/Engine/Source/Editor/VirtualizationEditor/Private/SVirtualAssetsStatistics.cpp
--- a/Engine/Source/Editor/VirtualizationEditor/Private/SVirtualAssetsStatistics.cpp
+++ b/Engine/Source/Editor/VirtualizationEditor/Private/SVirtualAssetsStatistics.cpp
@@ -234,297 +234,173 @@ TSharedRef<SWidget> SVirtualAssetsStatisticsDialog::GetGridPanel()
 	const float RowMargin = 0.0f;
 	const float TitleMargin = 10.0f;
 	const float ColumnMargin = 10.0f;
+	const float BorderPadding = ColumnMargin / 2.0f;
+
+	const FMargin StdMargin(ColumnMargin, RowMargin);
 	const FSlateColor TitleColor = FStyleColors::AccentWhite;
 	const FSlateFontInfo TitleFont = FCoreStyle::GetDefaultFontStyle("Bold", 10);
-	const double BytesToMegaBytes = 1.0 / (1024.0 * 1024.0);
-
-	TSharedRef<SGridPanel> Panel = SNew(SGridPanel);
+	
+	FSlateColor Color = FStyleColors::Foreground;
+	FSlateFontInfo Font = FCoreStyle::GetDefaultFontStyle("Regular", 10);
 
-	// TODO - need a way to make this work once the system is initialized
-	if (!IVirtualizationSystem::IsInitialized())
+	struct FPanels
 	{
-		Panel->AddSlot(0, 0)
-			[
-				SNew(STextBlock)
-				.Margin(FMargin(ColumnMargin, RowMargin))
-				.ColorAndOpacity(TitleColor)
-				.Font(TitleFont)
-				.Justification(ETextJustify::Center)
-				.Text(LOCTEXT("LazyInit", "System is uninitialized until first use"))
-			];
-
-		return Panel;
-	}
+		TSharedPtr<SGridPanel> Names;
+		TSharedPtr<SGridPanel> Pull;
+		TSharedPtr<SGridPanel> Cache;
+		TSharedPtr<SGridPanel> Push;
 
-	IVirtualizationSystem& System = IVirtualizationSystem::Get();	
-
-	if (System.IsEnabled() == false)
-	{
-		Panel->AddSlot(0,0)
-			[
-				SNew(STextBlock)
-				.Margin(FMargin(ColumnMargin, RowMargin))
-				.ColorAndOpacity(TitleColor)
-				.Font(TitleFont)
-				.Justification(ETextJustify::Center)
-				.Text(LOCTEXT("Disabled", "Virtual Assets Are Disabled For This Project"))
-			];
-	}
-	else
-	{
-		int32 Row = 0;
+	} Panels;
 
-		Panel->AddSlot(2, Row)
-			[
-				SNew(STextBlock)
-				.Margin(FMargin(ColumnMargin, RowMargin))
-				.ColorAndOpacity(TitleColor)
-				.Font(TitleFont)
-				.Justification(ETextJustify::Center)
-				.Text(LOCTEXT("Read", "Read"))
-			];
+	TSharedRef<SHorizontalBox> Panel = SNew(SHorizontalBox);
+	Panel->AddSlot()
+	.Padding(BorderPadding)
+	.AutoWidth()
+	[
+		SAssignNew(Panels.Names, SGridPanel)
+		+ SGridPanel::Slot(0, 0)
+		[
+			SNew(STextBlock)
+			.Margin(FMargin(ColumnMargin + (BorderPadding / 2.0f), RowMargin + (BorderPadding / 2.0f)))
+			.ColorAndOpacity(TitleColor)
+			.Font(TitleFont)
+			.Justification(ETextJustify::Left)
+		]
+		+SGridPanel::Slot(0, 1)
+		[
+			SNew(STextBlock)
+			.Margin(FMargin(ColumnMargin, RowMargin, 0.0f, TitleMargin))
+			.ColorAndOpacity(TitleColor)
+			.Font(TitleFont)
+			.Justification(ETextJustify::Left)
+			.Text(LOCTEXT("Backend", "Backend"))
+		]
+	];
 
-		Panel->AddSlot(5, Row)
-			[
-				SNew(STextBlock)
-				.Margin(FMargin(ColumnMargin, RowMargin))
-				.ColorAndOpacity(TitleColor)
-				.Font(TitleFont)
-				.Justification(ETextJustify::Center)
-				.Text(LOCTEXT("Write", "Write"))
-			];
+	auto CreateGridPanels = [&](FText&& Label)
+		{
+			TSharedPtr<SGridPanel> GridPanel;
 
-		Panel->AddSlot(8, Row)
+			Panel->AddSlot()
+			.Padding(BorderPadding)
+			.AutoWidth()
 			[
-				SNew(STextBlock)
-				.Margin(FMargin(ColumnMargin, RowMargin))
-				.ColorAndOpacity(TitleColor)
-				.Font(TitleFont)
-				.Justification(ETextJustify::Center)
-				.Text(LOCTEXT("Cache", "Cache"))
+				SNew(SBorder)
+				.Padding(BorderPadding)
+				[
+					SAssignNew(GridPanel, SGridPanel)
+					+ SGridPanel::Slot(1, 0)
+					[
+						SNew(STextBlock)
+						.Margin(FMargin(ColumnMargin, RowMargin))
+						.ColorAndOpacity(TitleColor)
+						.Font(TitleFont)
+						.Justification(ETextJustify::Center)
+						.Text(MoveTemp(Label))
+					]
+					+ SGridPanel::Slot(0, 1)
+					[
+						SNew(STextBlock)
+						.Margin(FMargin(ColumnMargin, RowMargin, 0.0f, TitleMargin))
+						.ColorAndOpacity(TitleColor)
+						.Font(TitleFont)
+						.Justification(ETextJustify::Center)
+						.Text(LOCTEXT("Count", "Count"))
+					]
+					+ SGridPanel::Slot(1, 1)
+					[
+						SNew(STextBlock)
+						.Margin(FMargin(ColumnMargin, RowMargin, 0.0f, TitleMargin))
+						.ColorAndOpacity(TitleColor)
+						.Font(TitleFont)
+						.Justification(ETextJustify::Center)
+						.Text(LOCTEXT("Size", "Size (MiB)"))
+					]
+					+ SGridPanel::Slot(2, 1)
+					[
+						SNew(STextBlock)
+						.Margin(FMargin(ColumnMargin, RowMargin, 0.0f, TitleMargin))
+						.ColorAndOpacity(TitleColor)
+						.Font(TitleFont)
+						.Justification(ETextJustify::Center)
+						.Text(LOCTEXT("Time", "Avg (ms)"))
+					]
+				]
 			];
+			
+			return GridPanel;
+		};
 
-		Row++;
+	Panels.Pull = CreateGridPanels(LOCTEXT("Download", "Download"));
+	Panels.Cache = CreateGridPanels(LOCTEXT("Cache", "Cache"));
+	Panels.Push = CreateGridPanels(LOCTEXT("Upload", "Upload"));
 
-		Panel->AddSlot(0, Row)
+	int32 RowIndex = 2;
+	auto DisplayPayloadActivityInfo = [&StdMargin, &Color, &Font, &RowIndex, &Panels](const FString& DebugName, const FString& ConfigName, const FPayloadActivityInfo& PayloadActivityInfo)
+		{
+			Panels.Names->AddSlot(0, RowIndex)
 			[
 				SNew(STextBlock)
-				.Margin(FMargin(ColumnMargin, RowMargin, 0.0f, TitleMargin))
-				.ColorAndOpacity(TitleColor)
-				.Font(TitleFont)
+				.Margin(StdMargin)
+				.ColorAndOpacity(Color)
+				.Font(Font)
 				.Justification(ETextJustify::Left)
-				.Text(LOCTEXT("Backend", "Backend"))
-			];
-
-		Panel->AddSlot(1, Row)
-			[
-				SNew(STextBlock)
-				.Margin(FMargin(ColumnMargin, RowMargin, 0.0f, TitleMargin))
-				.ColorAndOpacity(TitleColor)
-				.Font(TitleFont)
-				.Justification(ETextJustify::Center)
-				.Text(LOCTEXT("Count", "Count"))
+				.Text(FText::FromString(ConfigName))
 			];
 
-		Panel->AddSlot(2, Row)
-			[
-				SNew(STextBlock)
-				.Margin(FMargin(ColumnMargin, RowMargin, 0.0f, TitleMargin))
-				.ColorAndOpacity(TitleColor)
-				.Font(TitleFont)
-				.Justification(ETextJustify::Center)
-				.Text(LOCTEXT("Time", "Time (Sec)"))
-			];
-
-		Panel->AddSlot(3, Row)
-			[
-				SNew(STextBlock)
-				.Margin(FMargin(ColumnMargin, RowMargin, 0.0f, TitleMargin))
-				.ColorAndOpacity(TitleColor)
-				.Font(TitleFont)
-				.Justification(ETextJustify::Center)
-				.Text(LOCTEXT("Size", "Size (MB)"))
-			];
-
-		Panel->AddSlot(4, Row)
-			[
-				SNew(STextBlock)
-				.Margin(FMargin(ColumnMargin, RowMargin, 0.0f, TitleMargin))
-				.ColorAndOpacity(TitleColor)
-				.Font(TitleFont)
-				.Justification(ETextJustify::Center)
-				.Text(LOCTEXT("Count", "Count"))
-			];
-
-		Panel->AddSlot(5, Row)
-			[
-				SNew(STextBlock)
-				.Margin(FMargin(ColumnMargin, RowMargin, 0.0f, TitleMargin))
-				.ColorAndOpacity(TitleColor)
-				.Font(TitleFont)
-				.Justification(ETextJustify::Center)
-				.Text(LOCTEXT("Time", "Time (Sec)"))
-			];
-
-		Panel->AddSlot(6, Row)
-			[
-				SNew(STextBlock)
-				.Margin(FMargin(ColumnMargin, RowMargin, 0.0f, TitleMargin))
-				.ColorAndOpacity(TitleColor)
-				.Font(TitleFont)
-				.Justification(ETextJustify::Center)
-				.Text(LOCTEXT("Size", "Size (MB)"))
-			];
-
-		Panel->AddSlot(7, Row)
-			[
-				SNew(STextBlock)
-				.Margin(FMargin(ColumnMargin, RowMargin, 0.0f, TitleMargin))
-				.ColorAndOpacity(TitleColor)
-				.Font(TitleFont)
-				.Justification(ETextJustify::Center)
-				.Text(LOCTEXT("Count", "Count"))
-			];
-
-		Panel->AddSlot(8, Row)
-			[
-				SNew(STextBlock)
-				.Margin(FMargin(ColumnMargin, RowMargin, 0.0f, TitleMargin))
-				.ColorAndOpacity(TitleColor)
-				.Font(TitleFont)
-				.Justification(ETextJustify::Center)
-				.Text(LOCTEXT("Time", "Time (Sec)"))
-			];
-
-		Panel->AddSlot(9, Row)
-			[
-				SNew(STextBlock)
-				.Margin(FMargin(ColumnMargin, RowMargin, 0.0f, TitleMargin))
-				.ColorAndOpacity(TitleColor)
-				.Font(TitleFont)
-				.Justification(ETextJustify::Center)
-				.Text(LOCTEXT("Size", "Size (MB)"))
-			];
-
-		Row++;
-
-		FPayloadActivityInfo AccumulatedPayloadAcitvityInfo = System.GetAccumualtedPayloadActivityInfo();
-
-		FSlateColor Color = FStyleColors::Foreground;
-		FSlateFontInfo Font = FCoreStyle::GetDefaultFontStyle("Regular", 10);
-
-		auto DisplayPayloadActivityInfo = [&](const FString& DebugName, const FString& ConfigName, const FPayloadActivityInfo& PayloadActivityInfo)
-		{
-			Panel->AddSlot(0, Row)
-				[
-					SNew(STextBlock)
-					.Margin(FMargin(ColumnMargin, RowMargin))
-					.ColorAndOpacity(Color)
-					.Font(Font)
-					.Justification(ETextJustify::Left)
-					.Text(FText::FromString(DebugName))
-				];
-
-			Panel->AddSlot(1, Row)
-				[
-					SNew(STextBlock)
-					.Margin(FMargin(ColumnMargin, RowMargin))
-					.ColorAndOpacity(Color)
-					.Font(Font)
-					.Justification(ETextJustify::Center)
-					.Text_Lambda([PayloadActivityInfo] { return FText::FromString(FString::Printf(TEXT("%u"), PayloadActivityInfo.Pull.PayloadCount)); })
-				];
-
-			Panel->AddSlot(2, Row)
-				[
-					SNew(STextBlock)
-					.Margin(FMargin(ColumnMargin, RowMargin))
-					.ColorAndOpacity(Color)
-					.Font(Font)
-					.Justification(ETextJustify::Center)
-					.Text_Lambda([PayloadActivityInfo] { return FText::FromString(SingleDecimalFormat((double)PayloadActivityInfo.Pull.CyclesSpent * FPlatformTime::GetSecondsPerCycle())); })
-				];
-
-			Panel->AddSlot(3, Row)
-				[
-					SNew(STextBlock)
-					.Margin(FMargin(ColumnMargin, RowMargin))
-					.ColorAndOpacity(Color)
-					.Font(Font)
-					.Justification(ETextJustify::Center)
-					.Text_Lambda([PayloadActivityInfo, BytesToMegaBytes] { return FText::FromString(SingleDecimalFormat((double)PayloadActivityInfo.Pull.TotalBytes * BytesToMegaBytes)); })
-				];
-
-			Panel->AddSlot(4, Row)
-				[
-					SNew(STextBlock)
-					.Margin(FMargin(ColumnMargin, RowMargin))
-					.ColorAndOpacity(Color)
-					.Font(Font)
-					.Justification(ETextJustify::Center)
-					.Text_Lambda([PayloadActivityInfo] { return FText::FromString(FString::Printf(TEXT("%u"), PayloadActivityInfo.Push.PayloadCount)); })
-				];
-
-			Panel->AddSlot(5, Row)
-				[
-					SNew(STextBlock)
-					.Margin(FMargin(ColumnMargin, RowMargin))
-					.ColorAndOpacity(Color)
-					.Font(Font)
-					.Justification(ETextJustify::Center)
-					.Text_Lambda([PayloadActivityInfo] { return FText::FromString(SingleDecimalFormat((double)PayloadActivityInfo.Push.CyclesSpent * FPlatformTime::GetSecondsPerCycle())); })
-				];
-
-			Panel->AddSlot(6, Row)
-				[
-					SNew(STextBlock)
-					.Margin(FMargin(ColumnMargin, RowMargin))
-					.ColorAndOpacity(Color)
-					.Font(Font)
-					.Justification(ETextJustify::Center)
-					.Text_Lambda([PayloadActivityInfo, BytesToMegaBytes] { return FText::FromString(SingleDecimalFormat((double)PayloadActivityInfo.Push.TotalBytes * BytesToMegaBytes)); })
-				];
-
-			Panel->AddSlot(7, Row)
-				[
-					SNew(STextBlock)
-					.Margin(FMargin(ColumnMargin, RowMargin))
-					.ColorAndOpacity(Color)
-					.Font(Font)
-					.Justification(ETextJustify::Center)
-					.Text_Lambda([PayloadActivityInfo] { return FText::FromString(FString::Printf(TEXT("%u"), PayloadActivityInfo.Cache.PayloadCount)); })
-				];
-
-			Panel->AddSlot(8, Row)
-				[
-					SNew(STextBlock)
-					.Margin(FMargin(ColumnMargin, RowMargin))
-					.ColorAndOpacity(Color)
-					.Font(Font)
-					.Justification(ETextJustify::Center)
-					.Text_Lambda([PayloadActivityInfo] { return FText::FromString(SingleDecimalFormat((double)PayloadActivityInfo.Cache.CyclesSpent * FPlatformTime::GetSecondsPerCycle())); })
-				];
-
-			Panel->AddSlot(9, Row)
-				[
-					SNew(STextBlock)
-					.Margin(FMargin(ColumnMargin, RowMargin))
-					.ColorAndOpacity(Color)
-					.Font(Font)
-					.Justification(ETextJustify::Center)
-					.Text_Lambda([PayloadActivityInfo, BytesToMegaBytes] { return FText::FromString(SingleDecimalFormat((double)PayloadActivityInfo.Cache.TotalBytes * BytesToMegaBytes)); })
-				];
-
-			Row++;
+			auto FillPanelDetails = [&StdMargin, &Color, &Font, &RowIndex](TSharedPtr<SGridPanel>& Panel, const FPayloadActivityInfo::FActivity& Activity)
+				{
+					Panel->AddSlot(0, RowIndex)
+						[
+							SNew(STextBlock)
+							.Margin(StdMargin)
+							.ColorAndOpacity(Color)
+							.Font(Font)
+							.Justification(ETextJustify::Center)
+							.Text(FText::FromString(FString::Printf(TEXT("%u"), Activity.PayloadCount)))
+						];
+
+					const double TotalBytesMiB = Activity.TotalBytes / (1024.0 * 1024.0);
+
+					Panel->AddSlot(1, RowIndex)
+						[
+							SNew(STextBlock)
+							.Margin(StdMargin)
+							.ColorAndOpacity(Color)
+							.Font(Font)
+							.Justification(ETextJustify::Center)
+							.Text(FText::FromString(SingleDecimalFormat(TotalBytesMiB)))
+						];
+
+					const double TotalTime = static_cast<double>(FPlatformTime::ToMilliseconds64(Activity.CyclesSpent));
+					const double Avg = Activity.PayloadCount > 0 ? TotalTime / Activity.PayloadCount : 0.0;
+
+					Panel->AddSlot(2, RowIndex)
+						[
+							SNew(STextBlock)
+							.Margin(StdMargin)
+							.ColorAndOpacity(Color)
+							.Font(Font)
+							.Justification(ETextJustify::Center)
+							.Text(FText::FromString(SingleDecimalFormat(Avg)))
+						];
+				};
+
+			FillPanelDetails(Panels.Pull, PayloadActivityInfo.Pull);
+			FillPanelDetails(Panels.Cache, PayloadActivityInfo.Cache);
+			FillPanelDetails(Panels.Push, PayloadActivityInfo.Push);
+
+			RowIndex++;
 		};
 
-		System.GetPayloadActivityInfo(DisplayPayloadActivityInfo);
+	IVirtualizationSystem::Get().GetPayloadActivityInfo(DisplayPayloadActivityInfo);
 
-		Color = TitleColor;
-		Font = TitleFont;
+	FPayloadActivityInfo AccumulatedPayloadAcitvityInfo = IVirtualizationSystem::Get().GetAccumualtedPayloadActivityInfo();
 
-		DisplayPayloadActivityInfo(FString("Total"), FString("Total"), AccumulatedPayloadAcitvityInfo);
-	}
+	Color = TitleColor;
+	Font = TitleFont;
+
+	DisplayPayloadActivityInfo(FString("Total"), FString("Total"), AccumulatedPayloadAcitvityInfo);
 
 	return Panel;
 }
