Enable casting warnings for the VirtualizationEditor module

#rnx

diff --git a/Engine/Source/Editor/VirtualizationEditor/VirtualizationEditor.Build.cs b/Engine/Source/Editor/VirtualizationEditor/VirtualizationEditor.Build.cs
--- a/Engine/Source/Editor/VirtualizationEditor/VirtualizationEditor.Build.cs
+++ b/Engine/Source/Editor/VirtualizationEditor/VirtualizationEditor.Build.cs
@@ -7,6 +7,8 @@ public class VirtualizationEditor : ModuleRules
 	public VirtualizationEditor(ReadOnlyTargetRules Target)
 		 : base(Target)
 	{
+		UnsafeTypeCastWarningLevel = WarningLevel.Error;
+
 		PrivateDependencyModuleNames.AddRange(
 			new string[]
 			{
