Compile fix for clang cast warnings introduced in CL 33601457

#rnx

diff --git a/Engine/Source/Editor/VirtualizationEditor/Private/SVirtualAssetsStatistics.cpp b/Engine/Source/Editor/VirtualizationEditor/Private/SVirtualAssetsStatistics.cpp
--- a/Engine/Source/Editor/VirtualizationEditor/Private/SVirtualAssetsStatistics.cpp
+++ b/Engine/Source/Editor/VirtualizationEditor/Private/SVirtualAssetsStatistics.cpp
@@ -360,7 +360,7 @@ TSharedRef<SWidget> SVirtualAssetsStatisticsDialog::GetGridPanel()
 							.Text(FText::FromString(FString::Printf(TEXT("%u"), Activity.PayloadCount)))
 						];
 
-					const double TotalBytesMiB = Activity.TotalBytes / (1024.0 * 1024.0);
+					const double TotalBytesMiB = static_cast<double>(Activity.TotalBytes) / (1024.0 * 1024.0);
 
 					Panel->AddSlot(1, RowIndex)
 						[
@@ -373,7 +373,7 @@ TSharedRef<SWidget> SVirtualAssetsStatisticsDialog::GetGridPanel()
 						];
 
 					const double TotalTime = static_cast<double>(FPlatformTime::ToMilliseconds64(Activity.CyclesSpent));
-					const double Avg = Activity.PayloadCount > 0 ? TotalTime / Activity.PayloadCount : 0.0;
+					const double Avg = Activity.PayloadCount > 0 ? TotalTime / static_cast<double>(Activity.PayloadCount) : 0.0;
 
 					Panel->AddSlot(2, RowIndex)
 						[
