The VA statistics panel once again informs the user if the project does not have the VA system enabled.

#rnx

- If the virtualization system is not enabled for the current project then we should display a message informing the user of this rather than just displaying the normal statistics grid without any backends.
- This code was accidentally removed from the last refactor of this panel in CL 33601457, note that we also removed the lazyinit code path, this is fine as that code path is effectively deprecated anyway.

diff --git a/Engine/Source/Editor/VirtualizationEditor/Private/SVirtualAssetsStatistics.cpp b/Engine/Source/Editor/VirtualizationEditor/Private/SVirtualAssetsStatistics.cpp
--- a/Engine/Source/Editor/VirtualizationEditor/Private/SVirtualAssetsStatistics.cpp
+++ b/Engine/Source/Editor/VirtualizationEditor/Private/SVirtualAssetsStatistics.cpp
@@ -252,6 +252,17 @@ TSharedRef<SWidget> SVirtualAssetsStatisticsDialog::GetGridPanel()
 
 	} Panels;
 
+	// Early out if the system is disabled
+	if (IVirtualizationSystem::Get().IsEnabled() == false)
+	{
+		return	SNew(STextBlock)
+				.Margin(FMargin(ColumnMargin, RowMargin))
+				.ColorAndOpacity(TitleColor)
+				.Font(TitleFont)
+				.Justification(ETextJustify::Center)
+				.Text(LOCTEXT("Disabled", "Virtual Assets Are Disabled For This Project"));
+	}
+
 	TSharedRef<SHorizontalBox> Panel = SNew(SHorizontalBox);
 	Panel->AddSlot()
 	.Padding(BorderPadding)
