Improve the GeneratePayloadManifestCommandlet

#rnx

- Payload filter reasons generated during the virtualization process, such as min size and package path reasons will now be correctly recorded.
- The commandlet will now cleanup ALL files from the previous manifest rather than only replacing the files that it needs to write out. This should avoid confusion where files from different manifests were left in the same directory.
- Added new cmdline switches 'PendingOnly', 'FilteredOnly' and 'VirtualizedOnly' in addition to the preexisting 'LocalOnly' allowing finer control over what type of payloads are actually included in the manifest.
-- These options are now containers in a single enum rather than each as their own bool.

diff --git a/Engine/Source/Editor/VirtualizationEditor/Private/GeneratePayloadManifestCommandlet.cpp b/Engine/Source/Editor/VirtualizationEditor/Private/GeneratePayloadManifestCommandlet.cpp
--- a/Engine/Source/Editor/VirtualizationEditor/Private/GeneratePayloadManifestCommandlet.cpp
+++ b/Engine/Source/Editor/VirtualizationEditor/Private/GeneratePayloadManifestCommandlet.cpp
@@ -8,15 +8,28 @@
 #include "Misc/PackagePath.h"
 #include "Misc/Paths.h"
 #include "UObject/PackageTrailer.h"
+#include "VirtualizationExperimentalUtilities.h"
+
+ENUM_CLASS_FLAGS(UGeneratePayloadManifestCommandlet::EPayloadFilter);
 
 class FSpreadSheet
 {
 public:
+
 	bool OpenNewSheet()
 	{
+		if (NumSheets == 0)
+		{
+			if (!CleanupExistingFiles())
+			{
+				return false;
+			}
+		}
+
 		Ar.Reset();
 
-		const TStringBuilder<512> CSVPath = GetFilePath(++NumSheets);
+		TStringBuilder<512> CSVPath;
+		CreateFilePath(++NumSheets, CSVPath);
 
 		Ar = TUniquePtr<FArchive>(IFileManager::Get().CreateFileWriter(CSVPath.ToString()));
 		if (Ar.IsValid())
@@ -63,7 +76,10 @@ public:
 
 		for (int32 Index = 0; Index < NumSheets; ++Index)
 		{
-			Output << LINE_TERMINATOR << GetFilePath(Index);
+			TStringBuilder<512> Path;
+			CreateFilePath(Index, Path);
+
+			Output << LINE_TERMINATOR << Path;
 		}
 
 		return Output.ToString();
@@ -71,9 +87,29 @@ public:
 
 private:
 
-	static TStringBuilder<512> GetFilePath(int32 Index)
+	bool CleanupExistingFiles()
 	{
-		return WriteToString<512>(FPaths::ProjectSavedDir(), TEXT("PayloadManifest/sheet"), Index, TEXT(".csv"));
+		int32 SheetToDelete = 1;
+		TStringBuilder<512> Path;
+		CreateFilePath(SheetToDelete, Path);
+
+		while (IFileManager::Get().FileExists(Path.ToString()))
+		{
+			if (!IFileManager::Get().Delete(Path.ToString()))
+			{
+				return false;
+			}
+
+			CreateFilePath(++SheetToDelete, Path);
+		}
+
+		return true;
+	}
+
+	static void CreateFilePath(int32 Index, FStringBuilderBase& Path)
+	{
+		Path.Reset();
+		Path << FPaths::ProjectSavedDir() << TEXT("PayloadManifest/sheet") << Index << TEXT(".csv");
 	}
 
 	// With any values above these we had trouble importing the finished cvs into various spread sheet programs
@@ -129,20 +165,43 @@ int32 UGeneratePayloadManifestCommandlet::Main(const FString& Params)
 				{
 					PackageTrailerCount++;
 
-					Trailer.ForEachPayload([&Sheet, &PackagePath, &PayloadCount, bLocalOnly = bLocalOnly](const FIoHash& Id, uint64 SizeOnDisk, uint64 RawSize, UE::EPayloadAccessMode Mode, UE::Virtualization::EPayloadFilterReason Filter)->void
+					Trailer.ForEachPayload([&Sheet, &PackagePath, &PayloadCount, OutputFilter = Filter](const FIoHash& Id, uint64 SizeOnDisk, uint64 RawSize, UE::EPayloadAccessMode Mode, UE::Virtualization::EPayloadFilterReason Reason)->void
 						{
-							if (bLocalOnly && Mode != UE::EPayloadAccessMode::Local)
+							if (Mode != UE::EPayloadAccessMode::Virtualized)
+							{
+								if (EnumHasAllFlags(OutputFilter, EPayloadFilter::VirtualizedOnly))
+								{
+									return;
+								}
+
+								Reason = UE::Virtualization::Utils::FixFilterFlags(PackagePath, SizeOnDisk, Reason);
+
+								if (EnumHasAllFlags(OutputFilter, EPayloadFilter::PendingOnly) && Reason != UE::Virtualization::EPayloadFilterReason::None)
+								{
+									return;
+								}
+
+								if (EnumHasAllFlags(OutputFilter, EPayloadFilter::FilteredOnly) && Reason == UE::Virtualization::EPayloadFilterReason::None)
+								{
+									return;
+								}
+							}
+							else
 							{
-								return;
+								if (EnumHasAnyFlags(OutputFilter, EPayloadFilter::LocalOnly))
+								{
+									return;
+								}
 							}
 
 							PayloadCount++;
 
 							TAnsiStringBuilder<256> LineBuilder;
-							LineBuilder << PackagePath << "," << Id << "," << SizeOnDisk << "," << RawSize << "," << Mode << "," << *LexToString(Filter) <<"\n";
+							LineBuilder << PackagePath << "," << Id << "," << SizeOnDisk << "," << RawSize << "," << Mode << "," << *LexToString(Reason) <<"\n";
 							
 							Sheet.PrintRow(LineBuilder);
 						});
+
 				}
 			}
 
@@ -170,7 +229,37 @@ bool UGeneratePayloadManifestCommandlet::ParseCmdline(const FString& Params)
 
 	ParseCommandLine(*Params, Tokens, Switches);
 
-	bLocalOnly = Switches.Contains(TEXT("LocalOnly"));
+	auto ApplyFilter = [this, &Switches](const TCHAR* FilterName, EPayloadFilter FilterType)->bool
+		{
+			if (Switches.Contains(FilterName))
+			{
+				if (Filter != EPayloadFilter::None)
+				{
+					UE_LOG(LogVirtualization, Error, TEXT("Comandlet cannot have multiple filters active at the same time"));
+					return false;
+				}
+
+				Filter = FilterType;
+			}
+
+			return true;
+		};
+
+	TPair<const TCHAR*, EPayloadFilter> Filters[] =
+	{ 
+		{TEXT("LocalOnly"), EPayloadFilter::LocalOnly},
+		{TEXT("PendingOnly"), EPayloadFilter::PendingOnly},
+		{TEXT("FilteredOnly"), EPayloadFilter::FilteredOnly},
+		{TEXT("VirtualizedOnly"), EPayloadFilter::VirtualizedOnly}
+	};
+
+	for (const TPair<const TCHAR*, EPayloadFilter>& It : Filters)
+	{
+		if (!ApplyFilter(It.Key, It.Value))
+		{
+			return false;
+		}
+	}
 
 	return true;
 }
diff --git a/Engine/Source/Editor/VirtualizationEditor/Private/GeneratePayloadManifestCommandlet.h b/Engine/Source/Editor/VirtualizationEditor/Private/GeneratePayloadManifestCommandlet.h
--- a/Engine/Source/Editor/VirtualizationEditor/Private/GeneratePayloadManifestCommandlet.h
+++ b/Engine/Source/Editor/VirtualizationEditor/Private/GeneratePayloadManifestCommandlet.h
@@ -33,5 +33,15 @@ class UGeneratePayloadManifestCommandlet
 
 private:
 
-	bool bLocalOnly = false;
+	enum class EPayloadFilter
+	{
+		None = 0,
+		LocalOnly		=  1 << 0,
+		PendingOnly		= (1 << 1) | LocalOnly,
+		FilteredOnly	= (1 << 2) | LocalOnly,
+		VirtualizedOnly	= (1 << 3)
+	};
+	FRIEND_ENUM_CLASS_FLAGS(EPayloadFilter);
+
+	EPayloadFilter Filter = EPayloadFilter::None;
 };
