Fix minor issue where logging that long work was about to take place was actually logged after that work.

#rnx

diff --git a/Engine/Source/Editor/VirtualizationEditor/Private/CommandletUtils.cpp b/Engine/Source/Editor/VirtualizationEditor/Private/CommandletUtils.cpp
--- a/Engine/Source/Editor/VirtualizationEditor/Private/CommandletUtils.cpp
+++ b/Engine/Source/Editor/VirtualizationEditor/Private/CommandletUtils.cpp
@@ -140,6 +140,8 @@ TArray<FString> FindPackages(EFindPackageFlags Flags)
 
 TArray<FString> FindPackagesInDirectory(const FString& DirectoryToSearch)
 {
+	UE_LOG(LogVirtualization, Display, TEXT("Searching for packages under '%s'..."), *DirectoryToSearch);
+
 	TArray<FString> FilesInPackageFolder;
 	FPackageName::FindPackagesInDirectory(FilesInPackageFolder, DirectoryToSearch);
 
@@ -151,8 +153,6 @@ TArray<FString> FindPackagesInDirectory(const FString& DirectoryToSearch)
 		PackageNames.Add(FPaths::CreateStandardFilename(BasePath));
 	}
 
-	UE_LOG(LogVirtualization, Display, TEXT("Searching for packages under '%s'"), *DirectoryToSearch);
-
 	return PackageNames;
 }
 
