Add a new commandlet '-run="VirtualizationEditor.GenerateMountPointPayloadManifestCommandlet"' that generates a report on editor bulkdata payload usage grouped by mountpoint.

- Provides a good way to see the total size of structured data vs bulkdata, virtualized vs non virtualized, how much data could still be virtualized and how much is currently rejected by the projects virtualization filter options.
- Providing '-DetailedFilterReasons' on the cmdline will cause the commandlet to generate more detailed statistics on the payload filtering.
- The eventual plan for this commandlet is to provide an easy way to track disk usage for a project via our analytics systems. Especially for larger projects that are split up into many game feature plugins.

diff --git a/Engine/Source/Editor/VirtualizationEditor/Private/GenerateMountPointPayloadManifestCommandlet.cpp b/Engine/Source/Editor/VirtualizationEditor/Private/GenerateMountPointPayloadManifestCommandlet.cpp
new file mode 100644
--- /dev/null
+++ b/Engine/Source/Editor/VirtualizationEditor/Private/GenerateMountPointPayloadManifestCommandlet.cpp
@@ -0,0 +1,352 @@
+// Copyright Epic Games, Inc. All Rights Reserved.
+
+#include "GenerateMountPointPayloadManifestCommandlet.h"
+
+#include "Async/ParallelFor.h"
+#include "CommandletUtils.h"
+#include "HAL/FileManager.h"
+#include "UObject/PackageTrailer.h"
+#include "VirtualizationExperimentalUtilities.h"
+
+#if 1
+#define UE_OUTPUTBYTES(x) ((double)x / (1024.0 * 1024.0 * 1024.0))
+#else
+#define UE_OUTPUTBYTES(x)
+#endif
+
+struct FMountPointStatistics
+{
+	FString Name;
+
+	uint64 TotalFileSize = 0;
+	uint64 NumFiles = 0;
+	uint64 NumFilesWithPayloads = 0;
+	uint64 NumFilesWithPendingPayloads = 0;
+
+	uint64 PendingPayloadCount = 0;
+	uint64 FilteredPayloadCount = 0;
+
+	uint64 LocalPendingSize = 0;
+	uint64 LocalFilteredSize = 0;
+
+	/** Used when 'DetailedFilterReasons' cmdline switch is used */
+	struct FilteredDetails
+	{
+	public:
+		void AddFile(UE::Virtualization::EPayloadFilterReason Reason, uint64 FileSize)
+		{
+			if (Reason == UE::Virtualization::EPayloadFilterReason::None)
+			{
+				NumFiles[0]++;
+				TotalFileSize[0] += FileSize;
+			}
+			else
+			{
+				for (uint16 Index = 1; Index < UE::Virtualization::NumPayloadFilterReasons; ++Index)
+				{
+					if (((uint16)Reason & (1 << (Index-1))) != 0)
+					{
+						NumFiles[Index]++;
+						TotalFileSize[Index] += FileSize;
+					}
+				}
+			}
+		}
+
+		uint64 GetCount(uint16 FilterReasonIndex) const
+		{
+			check(FilterReasonIndex < UE::Virtualization::NumPayloadFilterReasons);
+			return NumFiles[FilterReasonIndex];
+		}
+
+		uint64 GetTotalSize(uint16 FilterReasonIndex) const
+		{
+			check(FilterReasonIndex < UE::Virtualization::NumPayloadFilterReasons);
+			return TotalFileSize[FilterReasonIndex];
+		}
+
+		FilteredDetails& operator+=(const FilteredDetails& Other)
+		{
+			for (uint16 Index = 0; Index < UE::Virtualization::NumPayloadFilterReasons; ++Index)
+			{
+				NumFiles[Index] += Other.NumFiles[Index];
+				TotalFileSize[Index] += Other.TotalFileSize[Index];
+			}
+			return *this;
+		}
+
+	private:
+		uint64 NumFiles[UE::Virtualization::NumPayloadFilterReasons] = { 0 };
+		uint64 TotalFileSize[UE::Virtualization::NumPayloadFilterReasons] = { 0 };
+	} FilteredDetails;
+
+	FMountPointStatistics& operator+=(const FMountPointStatistics& Other)
+	{
+		TotalFileSize += Other.TotalFileSize;
+		NumFiles += Other.NumFiles;
+		NumFilesWithPayloads += Other.NumFilesWithPayloads;
+		NumFilesWithPendingPayloads += Other.NumFilesWithPendingPayloads;
+
+		PendingPayloadCount += Other.PendingPayloadCount;
+		FilteredPayloadCount += Other.FilteredPayloadCount;
+
+		LocalPendingSize += Other.LocalPendingSize;
+		LocalFilteredSize += Other.LocalFilteredSize;
+
+		FilteredDetails += Other.FilteredDetails;
+
+		return *this;
+	}
+};
+
+FMountPointStatistics ProcessMountPoint(FStringView Name, const TArray<FString>& FilePaths)
+{
+	TRACE_CPUPROFILER_EVENT_SCOPE(ProcessMountPoint);
+
+	FMountPointStatistics Stats;
+	Stats.Name = Name;
+
+	TArray<FMountPointStatistics> ContextStats;
+
+	ParallelForWithTaskContext(ContextStats, FilePaths.Num(),
+		[&FilePaths](FMountPointStatistics& Stats, int32 Index)
+		{
+			const FString& FilePath = FilePaths[Index];
+			int64 FileSize = IFileManager::Get().FileSize(*FilePath);
+			if (FileSize == INDEX_NONE)
+			{
+				UE_LOG(LogVirtualization, Error, TEXT("Unable to find file '%s'"), *FilePath);
+				return;
+			}
+
+			Stats.NumFiles++;
+			Stats.TotalFileSize += FileSize;
+
+			UE::FPackageTrailer Trailer;
+			if (UE::FPackageTrailer::TryLoadFromFile(FilePath, Trailer))
+			{
+				if (Trailer.GetNumPayloads(UE::EPayloadStorageType::Any) > 0)
+				{
+					Stats.NumFilesWithPayloads++;
+				}
+
+				bool bHasPendingPayloads = false;
+
+				Trailer.ForEachPayload([&FilePath, &Stats, &bHasPendingPayloads](const FIoHash& Id, uint64 SizeOnDisk, uint64 RawSize, UE::EPayloadAccessMode Mode, UE::Virtualization::EPayloadFilterReason Filter)->void
+					{
+						if (Mode == UE::EPayloadAccessMode::Local)
+						{
+							Filter = UE::Virtualization::Utils::FixFilterFlags(FilePath, SizeOnDisk, Filter);
+
+							Stats.FilteredDetails.AddFile(Filter, SizeOnDisk);
+
+							if (Filter == UE::Virtualization::EPayloadFilterReason::None)
+							{
+								Stats.PendingPayloadCount++;
+								Stats.LocalPendingSize += SizeOnDisk;
+
+								bHasPendingPayloads = true;
+							}
+							else
+							{
+								Stats.FilteredPayloadCount++;
+								Stats.LocalFilteredSize += SizeOnDisk;
+							}
+						}
+					});
+
+				if (bHasPendingPayloads)
+				{
+					Stats.NumFilesWithPendingPayloads++;
+				}
+			}
+		});
+
+	for (const FMountPointStatistics& Context : ContextStats)
+	{
+		Stats += Context;
+	}
+
+	return Stats;
+}
+
+bool OutputMountPointStatistics(const TArray<FMountPointStatistics>& Statistics, bool bDetailedFilterReasons)
+{
+	TRACE_CPUPROFILER_EVENT_SCOPE(OutputMountPointStatistics);
+
+	TStringBuilder<512> OutputFilePath;
+	OutputFilePath << FPaths::ProjectSavedDir() << TEXT("PayloadManifest/mountpoints.csv");
+
+	TUniquePtr<FArchive> Ar = TUniquePtr<FArchive>(IFileManager::Get().CreateFileWriter(OutputFilePath.ToString()));
+	if (Ar.IsValid())
+	{
+		{
+			TAnsiStringBuilder<512> Heading;
+			Heading = "Name,NumFiles,NumFilesWithPayloads,PendingPayloadCount,FilteredPayloadCount,VirtualizedPercent,TotalFileSize,StructuredDataSize,PendingPayloadSize,FilteredPayloadSize";
+
+			if (bDetailedFilterReasons)
+			{
+				for (uint16 FilterIdx = 1; FilterIdx < UE::Virtualization::NumPayloadFilterReasons; ++FilterIdx)
+				{
+					Heading << ",Filter (" << LexToString((UE::Virtualization::EPayloadFilterReason)(1 << (FilterIdx-1))) << ") Size";
+				}
+			}
+
+			Heading << "\n";
+
+			Ar->Serialize((void*)Heading.GetData(), Heading.Len() * sizeof(FAnsiStringView::ElementType));
+
+			if (Ar->IsError())
+			{
+				UE_LOG(LogVirtualization, Error, TEXT("Failed to write out csv headers to '%s'"), OutputFilePath.ToString());
+				return false;
+			}
+		}
+
+		for (const FMountPointStatistics& Stats : Statistics)
+		{
+			const double VirtualizedPercent = ((double)(Stats.NumFiles - Stats.NumFilesWithPendingPayloads) / (double)Stats.NumFiles) * 100.0;
+			const uint64 StructuredDataSize = Stats.TotalFileSize - (Stats.LocalPendingSize + Stats.LocalFilteredSize);
+
+			TAnsiStringBuilder<256> Line;
+			Line << Stats.Name << ",";
+
+			Line << Stats.NumFiles << ",";
+			Line << Stats.NumFilesWithPayloads << ",";
+
+			Line << Stats.PendingPayloadCount << ",";
+			Line << Stats.FilteredPayloadCount << ",";
+
+			Line.Appendf("%.1f,", VirtualizedPercent);
+
+			Line << UE_OUTPUTBYTES(Stats.TotalFileSize) << ",";
+			Line << UE_OUTPUTBYTES(StructuredDataSize) << ",";
+			Line << UE_OUTPUTBYTES(Stats.LocalPendingSize) << ",";
+			Line << UE_OUTPUTBYTES(Stats.LocalFilteredSize);
+
+			if (bDetailedFilterReasons)
+			{
+				for (uint16 FilterIdx = 1; FilterIdx < UE::Virtualization::NumPayloadFilterReasons; ++FilterIdx)
+				{
+					Line << "," << UE_OUTPUTBYTES(Stats.FilteredDetails.GetTotalSize(FilterIdx));
+				}
+			}
+
+			Line << "\n";
+
+			Ar->Serialize((void*)Line.GetData(), Line.Len() * sizeof(FAnsiStringView::ElementType));
+
+			if (Ar->IsError())
+			{
+				UE_LOG(LogVirtualization, Error, TEXT("Failed to write csv data to '%s'"), OutputFilePath.ToString());
+				return false;
+			}
+		}
+
+		Ar.Reset();
+
+		UE_LOG(LogVirtualization, Display, TEXT("Wrote output to: '%s'"), OutputFilePath.ToString());
+		return true;
+	}
+	else
+	{
+		UE_LOG(LogVirtualization, Error, TEXT("Failed to open '%s' for writing"), OutputFilePath.ToString());
+		return false;
+	}
+}
+
+UGenerateMountPointPayloadManifestCommandlet::UGenerateMountPointPayloadManifestCommandlet(const FObjectInitializer& ObjectInitializer)
+	: Super(ObjectInitializer)
+{
+}
+
+int32 UGenerateMountPointPayloadManifestCommandlet::Main(const FString& Params)
+{
+	TRACE_CPUPROFILER_EVENT_SCOPE(GenerateMountPointPayloadManifestCommandlet);
+
+	if (!ParseCmdline(Params))
+	{
+		UE_LOG(LogVirtualization, Error, TEXT("Failed to parse the command line correctly"));
+		return -1;
+	}
+
+	UE_LOG(LogVirtualization, Display, TEXT("Generating mount point summary for all files..."));
+
+	TArray<FString> PackageNames = UE::Virtualization::DiscoverPackages(Params, UE::Virtualization::EFindPackageFlags::ExcludeEngineContent);
+
+	UE_LOG(LogVirtualization, Display, TEXT("Found %d files to look in"), PackageNames.Num());
+
+	TMap<FString, TArray<FString>> MountPointMap;
+
+	{
+		TRACE_CPUPROFILER_EVENT_SCOPE(SortingMountPoints);
+		UE_LOG(LogVirtualization, Display, TEXT("Sorting by mount point..."));
+
+		for (FString& Path : PackageNames)
+		{
+			int32 ContentIndex = Path.Find(TEXT("/Content/"));
+			if (ContentIndex != INDEX_NONE)
+			{
+				int32 MountPointIndex = Path.Find(TEXT("/"), ESearchCase::IgnoreCase, ESearchDir::FromEnd, ContentIndex);
+				if (MountPointIndex != INDEX_NONE)
+				{
+					MountPointIndex++;
+
+					FStringView MountPoint = FStringView(Path).SubStr(MountPointIndex, ContentIndex - MountPointIndex);
+
+					if (TArray<FString>* ExistingPaths = MountPointMap.FindByHash(GetTypeHash(MountPoint), MountPoint))
+					{
+						ExistingPaths->Add(MoveTemp(Path));
+					}
+					else
+					{
+						TArray<FString> NewPaths;
+						NewPaths.Add(MoveTemp(Path));
+
+						MountPointMap.Add(FString(MountPoint), MoveTemp(NewPaths));
+					}
+				}
+			}
+			else
+			{
+				UE_LOG(LogVirtualization, Warning, TEXT("Package '%s' not under a valid content directory, skipping!"), *Path);
+			}
+		}
+	}
+
+	UE_LOG(LogVirtualization, Display, TEXT("Found %d mountpoints"), MountPointMap.Num());
+
+	TArray<FMountPointStatistics> Stats;
+	Stats.Reserve(MountPointMap.Num());
+
+	UE_LOG(LogVirtualization, Display, TEXT("Processing mountpoints..."), MountPointMap.Num());
+	for (const TPair<FString, TArray<FString>>& It : MountPointMap)
+	{
+		FMountPointStatistics MountPointStats = ProcessMountPoint(It.Key, It.Value);
+		Stats.Add(MountPointStats);
+	}
+
+	UE_LOG(LogVirtualization, Display, TEXT("Processing mountpoints completed"), MountPointMap.Num());
+
+	if (!OutputMountPointStatistics(Stats, bDetailedFilterReasons))
+	{
+		return -1;
+	}
+
+	return 0;
+}
+
+bool UGenerateMountPointPayloadManifestCommandlet::ParseCmdline(const FString& Params)
+{
+	TArray<FString> Tokens;
+	TArray<FString> Switches;
+
+	ParseCommandLine(*Params, Tokens, Switches);
+
+	if (Switches.Contains(TEXT("DetailedFilterReasons")))
+	{
+		bDetailedFilterReasons = true;
+	}
+
+	return true;
+}
diff --git a/Engine/Source/Editor/VirtualizationEditor/Private/GenerateMountPointPayloadManifestCommandlet.h b/Engine/Source/Editor/VirtualizationEditor/Private/GenerateMountPointPayloadManifestCommandlet.h
new file mode 100644
--- /dev/null
+++ b/Engine/Source/Editor/VirtualizationEditor/Private/GenerateMountPointPayloadManifestCommandlet.h
@@ -0,0 +1,32 @@
+// Copyright Epic Games, Inc. All Rights Reserved.
+
+#pragma once
+
+#include "Commandlets/Commandlet.h"
+#include "UObject/ObjectMacros.h"
+
+#include "GenerateMountPointPayloadManifestCommandlet.generated.h"
+
+/**
+ * Because the commandlet is the VirtualizationEditor module it needs to be invoked 
+ * with the command line:
+ * -run="VirtualizationEditor.GenerateMountPointPayloadManifestCommandlet"
+ */
+UCLASS()
+class UGenerateMountPointPayloadManifestCommandlet
+	: public UCommandlet
+{
+	GENERATED_UCLASS_BODY()
+
+	//~ Begin UCommandlet Interface
+	virtual int32 Main(const FString& Params) override;
+	//~ End UCommandlet Interface
+
+	static int32 StaticMain(const FString& Params);
+
+private:
+
+	bool ParseCmdline(const FString& Params);
+
+	bool bDetailedFilterReasons = false;
+};
diff --git a/Engine/Source/Runtime/Core/Public/Virtualization/VirtualizationTypes.h b/Engine/Source/Runtime/Core/Public/Virtualization/VirtualizationTypes.h
--- a/Engine/Source/Runtime/Core/Public/Virtualization/VirtualizationTypes.h
+++ b/Engine/Source/Runtime/Core/Public/Virtualization/VirtualizationTypes.h
@@ -82,6 +82,7 @@ enum class EPayloadFilterReason : uint16
 	/** Filtered because the package is either a UMap or the owning editor bulkdata is under a UMapBuildDataRegistry */
 	MapContent = 1 << 4,
 };
+constexpr uint16 NumPayloadFilterReasons = 6;
 
 ENUM_CLASS_FLAGS(EPayloadFilterReason);
 
