Improvements for the VirtualizationEditor.CheckForVirtualizedContent commandlet

#rnx

- Skip cooked packages when gathering as they do not have FPackageTrailers and are of no interest to any commandlet in the VirtualizationEditor module.
- Fix typo in variable name bNoVAPContentInProject -> bNoVAContentInProject.
- The commandlet now reports an error if not input is specified on the commandline.
- Changes individual package errors to use structured logging to help improve the output display for CI systems.

diff --git a/Engine/Source/Editor/VirtualizationEditor/Private/CheckForVirtualizedContentCommandlet.cpp b/Engine/Source/Editor/VirtualizationEditor/Private/CheckForVirtualizedContentCommandlet.cpp
--- a/Engine/Source/Editor/VirtualizationEditor/Private/CheckForVirtualizedContentCommandlet.cpp
+++ b/Engine/Source/Editor/VirtualizationEditor/Private/CheckForVirtualizedContentCommandlet.cpp
@@ -4,6 +4,7 @@
 
 #include "CommandletUtils.h"
 #include "HAL/FileManager.h"
+#include "Logging/StructuredLog.h"
 #include "Misc/PackageName.h"
 #include "Misc/Paths.h"
 #include "UObject/PackageTrailer.h"
@@ -23,10 +24,11 @@ int32 UCheckForVirtualizedContentCommandlet::Main(const FString& Params)
 	ParseCommandLine(*Params, Tokens, Switches);
 
 	bool bNoVAContentInEngine = false;
-	bool bNoVAPContentInProject = false;
+	bool bNoVAContentInProject = false;
 
 	TArray<FString> DirectoriesToCheck;
 
+	int32 UnknownSwitch = 0;
 	for (const FString& Switch : Switches)
 	{
 		FString InputPath;
@@ -37,18 +39,28 @@ int32 UCheckForVirtualizedContentCommandlet::Main(const FString& Params)
 		}
 		else if (Switch == TEXT("CheckProject"))
 		{
-			bNoVAPContentInProject = true;
+			bNoVAContentInProject = true;
 		}
 		else if (FParse::Value(*Switch, TEXT("CheckDir="), InputPath))
 		{
 			InputPath.ParseIntoArray(DirectoriesToCheck, TEXT("+"), true);
 		}
+		else
+		{
+			UnknownSwitch++;
+		}
+	}
+
+	if (UnknownSwitch == Switches.Num())
+	{
+		UE_LOG(LogVirtualization, Error, TEXT("No input was provided for the commandlet. Use '-CheckEngine', '-CheckProject' or '-CheckDir=...'"));
+		return 2;
 	}
 
 	TArray<FString> EnginePackages;
 	TArray<FString> ProjectPackages;
 
-	if (bNoVAContentInEngine || bNoVAPContentInProject)
+	if (bNoVAContentInEngine || bNoVAContentInProject)
 	{
 		TRACE_CPUPROFILER_EVENT_SCOPE(SortAllPackages);
 
@@ -58,11 +70,12 @@ int32 UCheckForVirtualizedContentCommandlet::Main(const FString& Params)
 
 		for (const FString& Path : AllPackages)
 		{
-			if (bNoVAContentInEngine && Path.StartsWith(EngineDir))
+			const bool bIsEngineContent = Path.StartsWith(EngineDir);
+			if (bNoVAContentInEngine && bIsEngineContent)
 			{
 				EnginePackages.Add(Path);
 			}
-			else if (bNoVAPContentInProject)
+			else if (bNoVAContentInProject && !bIsEngineContent)
 			{
 				ProjectPackages.Add(Path);
 			}
@@ -79,7 +92,7 @@ int32 UCheckForVirtualizedContentCommandlet::Main(const FString& Params)
 		}
 	}
 
-	if (bNoVAPContentInProject)
+	if (bNoVAContentInProject)
 	{
 		if (!TryValidateContent(TEXT("Project"), ProjectPackages))
 		{
@@ -155,7 +168,7 @@ bool UCheckForVirtualizedContentCommandlet::TryValidateContent(const TCHAR* Debu
 	{
 		for (FString& Path : VirtualizedPackages)
 		{
-			UE_LOG(LogVirtualization, Error, TEXT("Package '%s' contains virtualized payloads"), *Path);
+			UE_LOGFMT(LogVirtualization, Error, "Package {PackagePath} contains virtualized payloads", Path);
 		}
 
 		UE_LOG(LogVirtualization, Error, TEXT("Found %d virtualized package(s) in %s content"), VirtualizedPackages.Num(), DebugName);
@@ -198,7 +211,7 @@ bool UCheckForVirtualizedContentCommandlet::TryValidateDirectory(const FString&
 	{
 		for (FString& Path : VirtualizedPackages)
 		{
-			UE_LOG(LogVirtualization, Error, TEXT("Package '%s' contains virtualized payloads"), *Path);
+			UE_LOGFMT(LogVirtualization, Error, "Package {PackagePath} contains virtualized payloads", Path);
 		}
 
 		UE_LOG(LogVirtualization, Error, TEXT("Found %d virtualized package(s) under '%s'"), VirtualizedPackages.Num(), *Directory);
diff --git a/Engine/Source/Editor/VirtualizationEditor/Private/CommandletUtils.cpp b/Engine/Source/Editor/VirtualizationEditor/Private/CommandletUtils.cpp
--- a/Engine/Source/Editor/VirtualizationEditor/Private/CommandletUtils.cpp
+++ b/Engine/Source/Editor/VirtualizationEditor/Private/CommandletUtils.cpp
@@ -118,6 +118,11 @@ TArray<FString> FindPackages(EFindPackageFlags Flags)
 
 	AssetRegistry.Get().EnumerateAllPackages([&PackagePaths, EnginePath, bFilterEngineContent](FName PackageName, const FAssetPackageData& PackageData)
 		{
+			if (PackageData.CookedHash.IsValid())
+			{
+				return; // Skip cooked packages as they won't have a FPackageTrailer
+			}
+	
 			FString RelFileName;
 			if (PackageData.Extension != EPackageExtension::Unspecified && PackageData.Extension != EPackageExtension::Custom)
 			{
