Add a new cmdline option to "VirtualizationEditor.GenerateMountPointPayloadManifestCommandlet" allowing the caller to choose how bytes are writen out to the csv file.

#rnx

- The new cmdline option is '-ByteFormat=Bytes/MiB/GiB' allowing the user to choose the format of the bytes. By default we output the raw bytes which is most useful for automated processes, the user can opt to output in MiB or GiB (depending on the scale of their content) to make a more human readable output if prefered.

diff --git a/Engine/Source/Editor/VirtualizationEditor/Private/GenerateMountPointPayloadManifestCommandlet.cpp b/Engine/Source/Editor/VirtualizationEditor/Private/GenerateMountPointPayloadManifestCommandlet.cpp
--- a/Engine/Source/Editor/VirtualizationEditor/Private/GenerateMountPointPayloadManifestCommandlet.cpp
+++ b/Engine/Source/Editor/VirtualizationEditor/Private/GenerateMountPointPayloadManifestCommandlet.cpp
@@ -5,6 +5,7 @@
 #include "Async/ParallelFor.h"
 #include "CommandletUtils.h"
 #include "HAL/FileManager.h"
+#include "Logging/StructuredLog.h"
 #include "UObject/PackageTrailer.h"
 #include "VirtualizationExperimentalUtilities.h"
 
@@ -170,7 +171,51 @@ FMountPointStatistics ProcessMountPoint(FStringView Name, const TArray<FString>&
 	return Stats;
 }
 
-bool OutputMountPointStatistics(const TArray<FMountPointStatistics>& Statistics, bool bDetailedFilterReasons)
+FAnsiStringBuilderBase& FormatByteOutput(FAnsiStringBuilderBase& Builder, uint64 Bytes, EByteFormatting ByteFormat)
+{
+	switch (ByteFormat)
+	{
+		case EByteFormatting::Bytes:
+			Builder << Bytes;
+			break;
+		case EByteFormatting::MiB:
+			Builder << (double)Bytes / (1024.0 * 1024.0);
+			break;
+		case EByteFormatting::GiB:
+			Builder << (double)Bytes / (1024.0 * 1024.0 * 1024.0);
+			break;
+		default:
+			checkNoEntry();
+
+	}
+
+	return Builder;
+}
+
+bool TryLexFromString(FStringView String, EByteFormatting& OutFormat)
+{
+	if (String == TEXT("Bytes"))
+	{
+		OutFormat = EByteFormatting::Bytes;
+		return true;
+	}
+	else if (String == TEXT("MiB"))
+	{
+		OutFormat = EByteFormatting::MiB;
+		return true;
+	}
+	else if (String == TEXT("GiB"))
+	{
+		OutFormat = EByteFormatting::GiB;
+		return true;
+	}
+	else
+	{
+		return false;
+	}
+}
+
+bool OutputMountPointStatistics(const TArray<FMountPointStatistics>& Statistics, bool bDetailedFilterReasons, EByteFormatting ByteFormat)
 {
 	TRACE_CPUPROFILER_EVENT_SCOPE(OutputMountPointStatistics);
 
@@ -188,7 +233,7 @@ bool OutputMountPointStatistics(const TArray<FMountPointStatistics>& Statistics,
 			{
 				for (uint16 FilterIdx = 1; FilterIdx < UE::Virtualization::NumPayloadFilterReasons; ++FilterIdx)
 				{
-					Heading << ",Filter (" << LexToString((UE::Virtualization::EPayloadFilterReason)(1 << (FilterIdx-1))) << ") Size";
+					Heading << ",Filter (" << LexToString((UE::Virtualization::EPayloadFilterReason)(1 << (FilterIdx - 1))) << ") Size";
 				}
 			}
 
@@ -219,16 +264,17 @@ bool OutputMountPointStatistics(const TArray<FMountPointStatistics>& Statistics,
 
 			Line.Appendf("%.1f,", VirtualizedPercent);
 
-			Line << UE_OUTPUTBYTES(Stats.TotalFileSize) << ",";
-			Line << UE_OUTPUTBYTES(StructuredDataSize) << ",";
-			Line << UE_OUTPUTBYTES(Stats.LocalPendingSize) << ",";
-			Line << UE_OUTPUTBYTES(Stats.LocalFilteredSize);
+			FormatByteOutput(Line, Stats.TotalFileSize, ByteFormat) << ",";
+			FormatByteOutput(Line, StructuredDataSize, ByteFormat) << ",";
+			FormatByteOutput(Line, Stats.LocalPendingSize, ByteFormat) << ",";
+			FormatByteOutput(Line, Stats.LocalFilteredSize, ByteFormat);
 
 			if (bDetailedFilterReasons)
 			{
 				for (uint16 FilterIdx = 1; FilterIdx < UE::Virtualization::NumPayloadFilterReasons; ++FilterIdx)
 				{
-					Line << "," << UE_OUTPUTBYTES(Stats.FilteredDetails.GetTotalSize(FilterIdx));
+					Line << ",";
+					FormatByteOutput(Line, Stats.FilteredDetails.GetTotalSize(FilterIdx), ByteFormat);
 				}
 			}
 
@@ -328,7 +374,7 @@ int32 UGenerateMountPointPayloadManifestCommandlet::Main(const FString& Params)
 
 	UE_LOG(LogVirtualization, Display, TEXT("Processing mountpoints completed"), MountPointMap.Num());
 
-	if (!OutputMountPointStatistics(Stats, bDetailedFilterReasons))
+	if (!OutputMountPointStatistics(Stats, bDetailedFilterReasons, ByteFormat))
 	{
 		return -1;
 	}
@@ -348,5 +394,15 @@ bool UGenerateMountPointPayloadManifestCommandlet::ParseCmdline(const FString& P
 		bDetailedFilterReasons = true;
 	}
 
+	FString ByteFormatString;
+	if (FParse::Value(*Params, TEXT("ByteFormat="), ByteFormatString))
+	{
+		if (!TryLexFromString(ByteFormatString, ByteFormat))
+		{
+			UE_LOGFMT(LogVirtualization, Error, "'{Arg}' is not a valid value for  -ByteFormat=", ByteFormatString);
+			return false;
+		}
+	}
+
 	return true;
 }
diff --git a/Engine/Source/Editor/VirtualizationEditor/Private/GenerateMountPointPayloadManifestCommandlet.h b/Engine/Source/Editor/VirtualizationEditor/Private/GenerateMountPointPayloadManifestCommandlet.h
--- a/Engine/Source/Editor/VirtualizationEditor/Private/GenerateMountPointPayloadManifestCommandlet.h
+++ b/Engine/Source/Editor/VirtualizationEditor/Private/GenerateMountPointPayloadManifestCommandlet.h
@@ -7,10 +7,24 @@
 
 #include "GenerateMountPointPayloadManifestCommandlet.generated.h"
 
+enum class EByteFormatting
+{
+	Bytes = 0,
+	MiB,
+	GiB
+};
+
 /**
  * Because the commandlet is the VirtualizationEditor module it needs to be invoked 
  * with the command line:
  * -run="VirtualizationEditor.GenerateMountPointPayloadManifestCommandlet"
+ * Additional args:
+ * "-DetailedFilterReasons"
+ *     This switch will provide a breakdown of how much content per mount point
+ *     is prevented from being virtualized by a specific filter.
+ * "-ByteFormat=Bytes/MiB/GiB"
+ *     This value allows the caller to set how bytes should be formatted to the
+ *     csv file. The default is to output raw bytes.
  */
 UCLASS()
 class UGenerateMountPointPayloadManifestCommandlet
@@ -29,4 +43,6 @@ private:
 	bool ParseCmdline(const FString& Params);
 
 	bool bDetailedFilterReasons = false;
+
+	EByteFormatting ByteFormat = EByteFormatting::Bytes;
 };
