Allow the final output of 'GenerateMountPointPayloadManifestCommandlet' to be redirected.

#rnx

- '-OutputPath=' allows the caller to provide a full path, filename and extension for the final output
- '-OutputName=' allows the caller to change the name of the final output but still have it write to the default directory
- Note that providing both args at the same time will result in an error.

diff --git a/Engine/Source/Editor/VirtualizationEditor/Private/GenerateMountPointPayloadManifestCommandlet.cpp b/Engine/Source/Editor/VirtualizationEditor/Private/GenerateMountPointPayloadManifestCommandlet.cpp
--- a/Engine/Source/Editor/VirtualizationEditor/Private/GenerateMountPointPayloadManifestCommandlet.cpp
+++ b/Engine/Source/Editor/VirtualizationEditor/Private/GenerateMountPointPayloadManifestCommandlet.cpp
@@ -1,4 +1,4 @@
-// Copyright Epic Games, Inc. All Rights Reserved.
+// Copyright Epic Games, Inccsv All Rights Reserved.
 
 #include "GenerateMountPointPayloadManifestCommandlet.h"
 
@@ -215,90 +215,85 @@ bool TryLexFromString(FStringView String, EByteFormatting& OutFormat)
 	}
 }
 
-bool OutputMountPointStatistics(const TArray<FMountPointStatistics>& Statistics, bool bDetailedFilterReasons, EByteFormatting ByteFormat)
+bool OutputMountPointStatistics(const FString& OutputFilePath, const TArray<FMountPointStatistics>& Statistics, bool bDetailedFilterReasons, EByteFormatting ByteFormat)
 {
 	TRACE_CPUPROFILER_EVENT_SCOPE(OutputMountPointStatistics);
 
-	TStringBuilder<512> OutputFilePath;
-	OutputFilePath << FPaths::ProjectSavedDir() << TEXT("PayloadManifest/mountpoints.csv");
+	TUniquePtr<FArchive> Ar = TUniquePtr<FArchive>(IFileManager::Get().CreateFileWriter(*OutputFilePath));
+	if (!Ar.IsValid())
+	{
+		UE_LOGFMT(LogVirtualization, Error, "Failed to open '{OutputFilePath}' for writing", OutputFilePath);
+		return false;
+	}
 
-	TUniquePtr<FArchive> Ar = TUniquePtr<FArchive>(IFileManager::Get().CreateFileWriter(OutputFilePath.ToString()));
-	if (Ar.IsValid())
 	{
-		{
-			TAnsiStringBuilder<512> Heading;
-			Heading = "Name,NumFiles,NumFilesWithPayloads,PendingPayloadCount,FilteredPayloadCount,VirtualizedPercent,TotalFileSize,StructuredDataSize,PendingPayloadSize,FilteredPayloadSize";
+		TAnsiStringBuilder<512> Heading;
+		Heading = "Name,NumFiles,NumFilesWithPayloads,PendingPayloadCount,FilteredPayloadCount,VirtualizedPercent,TotalFileSize,StructuredDataSize,PendingPayloadSize,FilteredPayloadSize";
 
-			if (bDetailedFilterReasons)
+		if (bDetailedFilterReasons)
+		{
+			for (uint16 FilterIdx = 1; FilterIdx < UE::Virtualization::NumPayloadFilterReasons; ++FilterIdx)
 			{
-				for (uint16 FilterIdx = 1; FilterIdx < UE::Virtualization::NumPayloadFilterReasons; ++FilterIdx)
-				{
-					Heading << ",Filter (" << LexToString((UE::Virtualization::EPayloadFilterReason)(1 << (FilterIdx - 1))) << ") Size";
-				}
+				Heading << ",Filter (" << LexToString((UE::Virtualization::EPayloadFilterReason)(1 << (FilterIdx - 1))) << ") Size";
 			}
+		}
 
-			Heading << "\n";
+		Heading << "\n";
 
-			Ar->Serialize((void*)Heading.GetData(), Heading.Len() * sizeof(FAnsiStringView::ElementType));
+		Ar->Serialize((void*)Heading.GetData(), Heading.Len() * sizeof(FAnsiStringView::ElementType));
 
-			if (Ar->IsError())
-			{
-				UE_LOG(LogVirtualization, Error, TEXT("Failed to write out csv headers to '%s'"), OutputFilePath.ToString());
-				return false;
-			}
+		if (Ar->IsError())
+		{
+			UE_LOGFMT(LogVirtualization, Error, "Failed to write out csv headers to '{OutputFilePath}'", OutputFilePath);
+			return false;
 		}
+	}
 
-		for (const FMountPointStatistics& Stats : Statistics)
-		{
-			const double VirtualizedPercent = ((double)(Stats.NumFiles - Stats.NumFilesWithPendingPayloads) / (double)Stats.NumFiles) * 100.0;
-			const uint64 StructuredDataSize = Stats.TotalFileSize - (Stats.LocalPendingSize + Stats.LocalFilteredSize);
+	for (const FMountPointStatistics& Stats : Statistics)
+	{
+		const double VirtualizedPercent = ((double)(Stats.NumFiles - Stats.NumFilesWithPendingPayloads) / (double)Stats.NumFiles) * 100.0;
+		const uint64 StructuredDataSize = Stats.TotalFileSize - (Stats.LocalPendingSize + Stats.LocalFilteredSize);
 
-			TAnsiStringBuilder<256> Line;
-			Line << Stats.Name << ",";
+		TAnsiStringBuilder<256> Line;
+		Line << Stats.Name << ",";
 
-			Line << Stats.NumFiles << ",";
-			Line << Stats.NumFilesWithPayloads << ",";
+		Line << Stats.NumFiles << ",";
+		Line << Stats.NumFilesWithPayloads << ",";
 
-			Line << Stats.PendingPayloadCount << ",";
-			Line << Stats.FilteredPayloadCount << ",";
+		Line << Stats.PendingPayloadCount << ",";
+		Line << Stats.FilteredPayloadCount << ",";
 
-			Line.Appendf("%.1f,", VirtualizedPercent);
+		Line.Appendf("%.1f,", VirtualizedPercent);
 
-			FormatByteOutput(Line, Stats.TotalFileSize, ByteFormat) << ",";
-			FormatByteOutput(Line, StructuredDataSize, ByteFormat) << ",";
-			FormatByteOutput(Line, Stats.LocalPendingSize, ByteFormat) << ",";
-			FormatByteOutput(Line, Stats.LocalFilteredSize, ByteFormat);
+		FormatByteOutput(Line, Stats.TotalFileSize, ByteFormat) << ",";
+		FormatByteOutput(Line, StructuredDataSize, ByteFormat) << ",";
+		FormatByteOutput(Line, Stats.LocalPendingSize, ByteFormat) << ",";
+		FormatByteOutput(Line, Stats.LocalFilteredSize, ByteFormat);
 
-			if (bDetailedFilterReasons)
+		if (bDetailedFilterReasons)
+		{
+			for (uint16 FilterIdx = 1; FilterIdx < UE::Virtualization::NumPayloadFilterReasons; ++FilterIdx)
 			{
-				for (uint16 FilterIdx = 1; FilterIdx < UE::Virtualization::NumPayloadFilterReasons; ++FilterIdx)
-				{
-					Line << ",";
-					FormatByteOutput(Line, Stats.FilteredDetails.GetTotalSize(FilterIdx), ByteFormat);
-				}
+				Line << ",";
+				FormatByteOutput(Line, Stats.FilteredDetails.GetTotalSize(FilterIdx), ByteFormat);
 			}
+		}
 
-			Line << "\n";
+		Line << "\n";
 
-			Ar->Serialize((void*)Line.GetData(), Line.Len() * sizeof(FAnsiStringView::ElementType));
+		Ar->Serialize((void*)Line.GetData(), Line.Len() * sizeof(FAnsiStringView::ElementType));
 
-			if (Ar->IsError())
-			{
-				UE_LOG(LogVirtualization, Error, TEXT("Failed to write csv data to '%s'"), OutputFilePath.ToString());
-				return false;
-			}
+		if (Ar->IsError())
+		{
+			UE_LOGFMT(LogVirtualization, Error, "Failed to write csv data to '{OutputFilePath}'", OutputFilePath);
+			return false;
 		}
+	}
 
-		Ar.Reset();
+	Ar.Reset();
 
-		UE_LOG(LogVirtualization, Display, TEXT("Wrote output to: '%s'"), OutputFilePath.ToString());
-		return true;
-	}
-	else
-	{
-		UE_LOG(LogVirtualization, Error, TEXT("Failed to open '%s' for writing"), OutputFilePath.ToString());
-		return false;
-	}
+	UE_LOGFMT(LogVirtualization, Display, "Wrote output to : '{OutputFilePath}'", OutputFilePath);
+	return true;
 }
 
 UGenerateMountPointPayloadManifestCommandlet::UGenerateMountPointPayloadManifestCommandlet(const FObjectInitializer& ObjectInitializer)
@@ -310,7 +305,7 @@ int32 UGenerateMountPointPayloadManifestCommandlet::Main(const FString& Params)
 {
 	TRACE_CPUPROFILER_EVENT_SCOPE(GenerateMountPointPayloadManifestCommandlet);
 
-	if (!ParseCmdline(Params))
+	if (!TryParseCmdline(Params))
 	{
 		UE_LOG(LogVirtualization, Error, TEXT("Failed to parse the command line correctly"));
 		return -1;
@@ -374,7 +369,7 @@ int32 UGenerateMountPointPayloadManifestCommandlet::Main(const FString& Params)
 
 	UE_LOG(LogVirtualization, Display, TEXT("Processing mountpoints completed"), MountPointMap.Num());
 
-	if (!OutputMountPointStatistics(Stats, bDetailedFilterReasons, ByteFormat))
+	if (!OutputMountPointStatistics(OutputFilePath, Stats, bDetailedFilterReasons, ByteFormat))
 	{
 		return -1;
 	}
@@ -382,7 +377,7 @@ int32 UGenerateMountPointPayloadManifestCommandlet::Main(const FString& Params)
 	return 0;
 }
 
-bool UGenerateMountPointPayloadManifestCommandlet::ParseCmdline(const FString& Params)
+bool UGenerateMountPointPayloadManifestCommandlet::TryParseCmdline(const FString& Params)
 {
 	TArray<FString> Tokens;
 	TArray<FString> Switches;
@@ -404,5 +399,44 @@ bool UGenerateMountPointPayloadManifestCommandlet::ParseCmdline(const FString& P
 		}
 	}
 
+	if (!TryParseOutputPathFromCmdline(Params))
+	{
+		return false;
+	}
+
+	return true;
+}
+
+bool UGenerateMountPointPayloadManifestCommandlet::TryParseOutputPathFromCmdline(const FString& Params)
+{
+	// First see the entire file path has been provided on the commandline
+	FParse::Value(*Params, TEXT("OutputPath="), OutputFilePath);
+
+	// Now check to see if just the filename has been provided on the commandline
+	FString CustomOutputName;
+	if (FParse::Value(*Params, TEXT("OutputName="), CustomOutputName))
+	{
+		if (!OutputFilePath.IsEmpty() && !CustomOutputName.IsEmpty())
+		{
+			UE_LOGFMT(LogVirtualization, Error, "It is invalid to set both '-OutputPath=' and '-OutputName=' for the same command");
+			return false;
+		}
+
+		TStringBuilder<512> Builder;
+		Builder << FPaths::ProjectSavedDir() << TEXT("PayloadManifest/") << CustomOutputName;
+
+		OutputFilePath = Builder.ToView();
+	}
+
+	if (OutputFilePath.IsEmpty())
+	{
+		TStringBuilder<512> Builder;
+		Builder << FPaths::ProjectSavedDir() << TEXT("PayloadManifest/mountpoints.csv");
+
+		OutputFilePath = Builder.ToView();
+	}
+
+	UE_LOGFMT(LogVirtualization, Display, "Final output will be written to {OutputFilePath}", OutputFilePath);
+
 	return true;
 }
diff --git a/Engine/Source/Editor/VirtualizationEditor/Private/GenerateMountPointPayloadManifestCommandlet.h b/Engine/Source/Editor/VirtualizationEditor/Private/GenerateMountPointPayloadManifestCommandlet.h
--- a/Engine/Source/Editor/VirtualizationEditor/Private/GenerateMountPointPayloadManifestCommandlet.h
+++ b/Engine/Source/Editor/VirtualizationEditor/Private/GenerateMountPointPayloadManifestCommandlet.h
@@ -18,6 +18,10 @@ enum class EByteFormatting
  * Because the commandlet is the VirtualizationEditor module it needs to be invoked 
  * with the command line:
  * -run="VirtualizationEditor.GenerateMountPointPayloadManifestCommandlet"
+ * 
+ * By default the final output will be written to:
+  <project root>/saved/PayloadManifest/mountpoints.csv
+ * 
  * Additional args:
  * "-DetailedFilterReasons"
  *     This switch will provide a breakdown of how much content per mount point
@@ -25,6 +29,12 @@ enum class EByteFormatting
  * "-ByteFormat=Bytes/MiB/GiB"
  *     This value allows the caller to set how bytes should be formatted to the
  *     csv file. The default is to output raw bytes.
+ * "-OutputPath="
+ *     The full path (including filename and extension) of where to write the
+ *     final output.
+ * "-OutputName="
+ *     The file name (including extension) of where to write the final output.
+ *     Note that the file will be written to the default output directory.
  */
 UCLASS()
 class UGenerateMountPointPayloadManifestCommandlet
@@ -40,9 +50,11 @@ class UGenerateMountPointPayloadManifestCommandlet
 
 private:
 
-	bool ParseCmdline(const FString& Params);
+	bool TryParseCmdline(const FString& Params);
+	bool TryParseOutputPathFromCmdline(const FString& Params);
 
-	bool bDetailedFilterReasons = false;
+	FString OutputFilePath;
 
+	bool bDetailedFilterReasons = false;
 	EByteFormatting ByteFormat = EByteFormatting::Bytes;
 };
