Remove typo that got introduced in the copyright boilerplate code.

#rnx

diff --git a/Engine/Source/Editor/VirtualizationEditor/Private/GenerateMountPointPayloadManifestCommandlet.cpp b/Engine/Source/Editor/VirtualizationEditor/Private/GenerateMountPointPayloadManifestCommandlet.cpp
--- a/Engine/Source/Editor/VirtualizationEditor/Private/GenerateMountPointPayloadManifestCommandlet.cpp
+++ b/Engine/Source/Editor/VirtualizationEditor/Private/GenerateMountPointPayloadManifestCommandlet.cpp
@@ -1,4 +1,4 @@
-// Copyright Epic Games, Inccsv All Rights Reserved.
+// Copyright Epic Games, Inc. All Rights Reserved.
 
 #include "GenerateMountPointPayloadManifestCommandlet.h"
 
