Improve commandline input sanitization for GenerateMountPointPayloadManifest commandlet

#rnx

- Normalize the OutputPath/Name provided by the user so that our directory separators are consistent.

diff --git a/Engine/Source/Editor/VirtualizationEditor/Private/GenerateMountPointPayloadManifestCommandlet.cpp b/Engine/Source/Editor/VirtualizationEditor/Private/GenerateMountPointPayloadManifestCommandlet.cpp
--- a/Engine/Source/Editor/VirtualizationEditor/Private/GenerateMountPointPayloadManifestCommandlet.cpp
+++ b/Engine/Source/Editor/VirtualizationEditor/Private/GenerateMountPointPayloadManifestCommandlet.cpp
@@ -436,6 +436,8 @@ bool UGenerateMountPointPayloadManifestCommandlet::TryParseOutputPathFromCmdline
 		OutputFilePath = Builder.ToView();
 	}
 
+	FPaths::NormalizeFilename(OutputFilePath);
+
 	UE_LOGFMT(LogVirtualization, Display, "Final output will be written to {OutputFilePath}", OutputFilePath);
 
 	return true;
