Improve code documentation for the ValidateVirtualizedContent commandlet by exposing the optional command line switches.

#rnx

diff --git a/Engine/Source/Editor/VirtualizationEditor/Private/ValidateVirtualizedContentCommandlet.h b/Engine/Source/Editor/VirtualizationEditor/Private/ValidateVirtualizedContentCommandlet.h
--- a/Engine/Source/Editor/VirtualizationEditor/Private/ValidateVirtualizedContentCommandlet.h
+++ b/Engine/Source/Editor/VirtualizationEditor/Private/ValidateVirtualizedContentCommandlet.h
@@ -23,6 +23,12 @@ struct FIoHash;
  * Because the commandlet is the VirtualizationEditor module it needs to be invoked
  * with the command line:
  * -run="VirtualizationEditor.ValidateVirtualizedContent"
+ * 
+ * Optional Args:
+ * bValidateContent:	When this switch is set each payload will be pulled from persistent storage
+ *						and additional validations will be run to make sure that the payload is
+ *						not corrupted. This is much slower than just checking that the payload
+ *						exists.
  */
 UCLASS()
 class UValidateVirtualizedContentCommandlet
