Improvements to the 'CheckForVirtualizedContent' commandlet

#rnx

- When running with '-CheckEngine' we now do a findinfiles on the root Engine directory to find all package in the engine. Previously we were using the asset registry which would only find content for engine plugins that are active for the current project, now we find the content for all engine plugins as well as engine content.
- When running with '-CheckProject' we now filter out the engine content at the assetregistry level.
- As a bonus this also means that we can split the package finding code into engine and project code blocks, which should make it easier to follow.

diff --git a/Engine/Source/Editor/VirtualizationEditor/Private/CheckForVirtualizedContentCommandlet.cpp b/Engine/Source/Editor/VirtualizationEditor/Private/CheckForVirtualizedContentCommandlet.cpp
--- a/Engine/Source/Editor/VirtualizationEditor/Private/CheckForVirtualizedContentCommandlet.cpp
+++ b/Engine/Source/Editor/VirtualizationEditor/Private/CheckForVirtualizedContentCommandlet.cpp
@@ -62,26 +62,22 @@ int32 UCheckForVirtualizedContentCommandlet::Main(const FString& Params)
 	TArray<FString> EnginePackages;
 	TArray<FString> ProjectPackages;
 
-	if (bNoVAContentInEngine || bNoVAContentInProject)
+	if (bNoVAContentInEngine)
 	{
-		TRACE_CPUPROFILER_EVENT_SCOPE(SortAllPackages);
+		TRACE_CPUPROFILER_EVENT_SCOPE(FindEnginePackages);
 
-		TArray<FString> AllPackages = UE::Virtualization::FindPackages(UE::Virtualization::EFindPackageFlags::None);
+		// When checking engine content it is assumed that we want to check ALL engine content not just the plugins that
+		// happen to be enabled for the current project. This is why we cannot use the asset registry by calling
+		// UE::Virtualization::FindPackages as that will only find th engine content enabled for the current project.
+		// Instead we search the engine from it's root directory.
+		EnginePackages =  UE::Virtualization::FindPackagesInDirectory(FPaths::EngineDir());
+	}
 
-		const FString EngineDir = FPaths::EngineDir();
+	if (bNoVAContentInProject)
+	{
+		TRACE_CPUPROFILER_EVENT_SCOPE(FindProjectPackages);
 
-		for (const FString& Path : AllPackages)
-		{
-			const bool bIsEngineContent = Path.StartsWith(EngineDir);
-			if (bNoVAContentInEngine && bIsEngineContent)
-			{
-				EnginePackages.Add(Path);
-			}
-			else if (bNoVAContentInProject && !bIsEngineContent)
-			{
-				ProjectPackages.Add(Path);
-			}
-		}
+		ProjectPackages = UE::Virtualization::FindPackages(UE::Virtualization::EFindPackageFlags::ExcludeEngineContent);
 	}
 
 	bool bAllContentValid = true;
