Add new cmdline option 'OutputPaths' to the 'CheckForVirtualizedContent' commandlet.

#rnx

- Sometimes we want to just get the path of the package rather than it's package name since the package name is based off the mount point and we might not know where those mount points are making the package files hard to find.
- When the new cmdline option 'OutputPaths' is supplied to the commandlet we will log virtualized content errors with the package path instead.
-- Default behavior will remain unchanged and we will convert paths to package names where possible.
- Relative paths will now be converted to absolute paths.
- Now that we have cmdline options which are not input, the logic for detecting that no input switches were provided needed to be redone.

#Synced-CL 48446423

diff --git a/Engine/Source/Editor/VirtualizationEditor/Private/CheckForVirtualizedContentCommandlet.cpp b/Engine/Source/Editor/VirtualizationEditor/Private/CheckForVirtualizedContentCommandlet.cpp
--- a/Engine/Source/Editor/VirtualizationEditor/Private/CheckForVirtualizedContentCommandlet.cpp
+++ b/Engine/Source/Editor/VirtualizationEditor/Private/CheckForVirtualizedContentCommandlet.cpp
@@ -30,7 +30,8 @@ int32 UCheckForVirtualizedContentCommandlet::Main(const FString& Params)
 
 	TArray<FString> DirectoriesToCheck;
 
-	int32 UnknownSwitch = 0;
+	bool bInputProvided = false;
+
 	for (const FString& Switch : Switches)
 	{
 		FString InputPath;
@@ -38,22 +39,25 @@ int32 UCheckForVirtualizedContentCommandlet::Main(const FString& Params)
 		if (Switch == TEXT("CheckEngine"))
 		{
 			bNoVAContentInEngine = true;
+			bInputProvided = true;
 		}
 		else if (Switch == TEXT("CheckProject"))
 		{
 			bNoVAContentInProject = true;
+			bInputProvided = true;
 		}
 		else if (FParse::Value(*Switch, TEXT("CheckDir="), InputPath))
 		{
 			InputPath.ParseIntoArray(DirectoriesToCheck, TEXT("+"), true);
+			bInputProvided = true;
 		}
-		else
+		else if (Switch == TEXT("OutputPaths"))
 		{
-			UnknownSwitch++;
+			bOutputPackageNames = false;
 		}
 	}
 
-	if (UnknownSwitch == Switches.Num())
+	if (bInputProvided == false)
 	{
 		UE_LOG(LogVirtualization, Error, TEXT("No input was provided for the commandlet. Use '-CheckEngine', '-CheckProject' or '-CheckDir=...'"));
 		return 2;
@@ -129,13 +133,20 @@ TArray<FString> UCheckForVirtualizedContentCommandlet::FindVirtualizedPackages(c
 			if (NumVirtualizedPayloads > 0)
 			{
 				FString PackageName;
-				if (FPackageName::TryConvertFilenameToLongPackageName(Path, PackageName))
+				if (bOutputPackageNames && FPackageName::TryConvertFilenameToLongPackageName(Path, PackageName))
 				{
 					VirtualizedPackages.Emplace(MoveTemp(PackageName));
 				}
 				else
 				{
-					VirtualizedPackages.Add(Path);
+					if (FPaths::IsRelative(Path))
+					{
+						VirtualizedPackages.Add(FPaths::ConvertRelativePathToFull(Path));
+					}
+					else
+					{
+						VirtualizedPackages.Add(Path);
+					}
 				}
 
 			}
diff --git a/Engine/Source/Editor/VirtualizationEditor/Private/CheckForVirtualizedContentCommandlet.h b/Engine/Source/Editor/VirtualizationEditor/Private/CheckForVirtualizedContentCommandlet.h
--- a/Engine/Source/Editor/VirtualizationEditor/Private/CheckForVirtualizedContentCommandlet.h
+++ b/Engine/Source/Editor/VirtualizationEditor/Private/CheckForVirtualizedContentCommandlet.h
@@ -18,6 +18,8 @@
  * more than one path)
  *     Checks packages in the given directory and all subdirectories
  * any or all of which can be passed to the commandlet.
+ * -OutputPaths
+ *     Report virtualized packages by path rather than package name.
  *
  * If virtualized payloads are found then the package path (or file path if CheckDir
  * indicates a directory that is not currently mounted by the project) is logged
@@ -48,4 +50,6 @@ private:
 	bool TryValidateContent(const TCHAR* DebugName, const TArray<FString>& Packages);
 	bool TryValidateDirectory(const FString& Directory);
 
+	bool bOutputPackageNames = true;
+
 };
