Add the mount point paths to the output of GenerateMountPointPayloadManifestCommandlet

#rnx

- Instead of storing the mountpoint name we now store the path, when we output the data to the csv file we can easily extract the name.
- Instead of passing bools to ::OutputMountPointStatistics a new enum 'EOutputFlags has been added to pass in the options.
- By default we will write out the path of the mount point in the last column, this can be supressed by using the cmdline arg '-NoPaths'

#Synced-CL 49468340

diff --git a/Engine/Source/Editor/VirtualizationEditor/Private/GenerateMountPointPayloadManifestCommandlet.cpp b/Engine/Source/Editor/VirtualizationEditor/Private/GenerateMountPointPayloadManifestCommandlet.cpp
--- a/Engine/Source/Editor/VirtualizationEditor/Private/GenerateMountPointPayloadManifestCommandlet.cpp
+++ b/Engine/Source/Editor/VirtualizationEditor/Private/GenerateMountPointPayloadManifestCommandlet.cpp
@@ -6,20 +6,30 @@
 #include "CommandletUtils.h"
 #include "HAL/FileManager.h"
 #include "Logging/StructuredLog.h"
+#include "Misc/PathViews.h"
 #include "UObject/PackageTrailer.h"
 #include "VirtualizationExperimentalUtilities.h"
 
 #include UE_INLINE_GENERATED_CPP_BY_NAME(GenerateMountPointPayloadManifestCommandlet)
 
 #if 1
-#define UE_OUTPUTBYTES(x) ((double)x / (1024.0 * 1024.0 * 1024.0))
+	#define UE_OUTPUTBYTES(x) ((double)x / (1024.0 * 1024.0 * 1024.0))
 #else
-#define UE_OUTPUTBYTES(x)
+	#define UE_OUTPUTBYTES(x)
 #endif
 
+enum class EOutputFlags : uint32
+{
+	None					= 0,
+	DetailedFilterReasons	= 1 << 0,
+	MountPointPaths			= 1 << 1
+};
+
+ENUM_CLASS_FLAGS(EOutputFlags);
+
 struct FMountPointStatistics
 {
-	FString Name;
+	FString RootPath;
 
 	uint64 TotalFileSize = 0;
 	uint64 NumFiles = 0;
@@ -102,12 +112,12 @@ struct FMountPointStatistics
 	}
 };
 
-FMountPointStatistics ProcessMountPoint(FStringView Name, const TArray<FString>& FilePaths)
+FMountPointStatistics ProcessMountPoint(FStringView RootPath, const TArray<FString>& FilePaths)
 {
 	TRACE_CPUPROFILER_EVENT_SCOPE(ProcessMountPoint);
 
 	FMountPointStatistics Stats;
-	Stats.Name = Name;
+	Stats.RootPath = RootPath;
 
 	TArray<FMountPointStatistics> ContextStats;
 
@@ -217,7 +227,7 @@ bool TryLexFromString(FStringView String, EByteFormatting& OutFormat)
 	}
 }
 
-bool OutputMountPointStatistics(const FString& OutputFilePath, const TArray<FMountPointStatistics>& Statistics, bool bDetailedFilterReasons, EByteFormatting ByteFormat)
+bool OutputMountPointStatistics(const FString& OutputFilePath, const TArray<FMountPointStatistics>& Statistics, EOutputFlags Flags, EByteFormatting ByteFormat)
 {
 	TRACE_CPUPROFILER_EVENT_SCOPE(OutputMountPointStatistics);
 
@@ -232,7 +242,7 @@ bool OutputMountPointStatistics(const FString& OutputFilePath, const TArray<FMou
 		TAnsiStringBuilder<512> Heading;
 		Heading = "Name,NumFiles,NumFilesWithPayloads,PendingPayloadCount,FilteredPayloadCount,VirtualizedPercent,TotalFileSize,StructuredDataSize,PendingPayloadSize,FilteredPayloadSize";
 
-		if (bDetailedFilterReasons)
+		if (EnumHasAllFlags(Flags, EOutputFlags::DetailedFilterReasons))
 		{
 			for (uint16 FilterIdx = 1; FilterIdx < UE::Virtualization::NumPayloadFilterReasons; ++FilterIdx)
 			{
@@ -240,6 +250,11 @@ bool OutputMountPointStatistics(const FString& OutputFilePath, const TArray<FMou
 			}
 		}
 
+		if (EnumHasAllFlags(Flags, EOutputFlags::MountPointPaths))
+		{
+			Heading << ",Path";
+		}
+
 		Heading << "\n";
 
 		Ar->Serialize((void*)Heading.GetData(), Heading.Len() * sizeof(FAnsiStringView::ElementType));
@@ -253,11 +268,11 @@ bool OutputMountPointStatistics(const FString& OutputFilePath, const TArray<FMou
 
 	for (const FMountPointStatistics& Stats : Statistics)
 	{
-		const double VirtualizedPercent = ((double)(Stats.NumFiles - Stats.NumFilesWithPendingPayloads) / (double)Stats.NumFiles) * 100.0;
+		const double VirtualizedPercent = Stats.NumFiles > 0 ? ((double)(Stats.NumFiles - Stats.NumFilesWithPendingPayloads) / (double)Stats.NumFiles) * 100.0 : 0.0;
 		const uint64 StructuredDataSize = Stats.TotalFileSize - (Stats.LocalPendingSize + Stats.LocalFilteredSize);
 
 		TAnsiStringBuilder<256> Line;
-		Line << Stats.Name << ",";
+		Line << "\"" << FPathViews::GetPathLeaf(Stats.RootPath) << "\",";
 
 		Line << Stats.NumFiles << ",";
 		Line << Stats.NumFilesWithPayloads << ",";
@@ -272,7 +287,7 @@ bool OutputMountPointStatistics(const FString& OutputFilePath, const TArray<FMou
 		FormatByteOutput(Line, Stats.LocalPendingSize, ByteFormat) << ",";
 		FormatByteOutput(Line, Stats.LocalFilteredSize, ByteFormat);
 
-		if (bDetailedFilterReasons)
+		if (EnumHasAllFlags(Flags, EOutputFlags::DetailedFilterReasons))
 		{
 			for (uint16 FilterIdx = 1; FilterIdx < UE::Virtualization::NumPayloadFilterReasons; ++FilterIdx)
 			{
@@ -281,6 +296,11 @@ bool OutputMountPointStatistics(const FString& OutputFilePath, const TArray<FMou
 			}
 		}
 
+		if (EnumHasAllFlags(Flags, EOutputFlags::MountPointPaths))
+		{
+			Line << ",\"" << Stats.RootPath << "\"";
+		}
+
 		Line << "\n";
 
 		Ar->Serialize((void*)Line.GetData(), Line.Len() * sizeof(FAnsiStringView::ElementType));
@@ -319,6 +339,7 @@ int32 UGenerateMountPointPayloadManifestCommandlet::Main(const FString& Params)
 
 	UE_LOG(LogVirtualization, Display, TEXT("Found %d files to look in"), PackageNames.Num());
 
+	// RootPath of the mount point, array of file paths under that mount point
 	TMap<FString, TArray<FString>> MountPointMap;
 
 	{
@@ -335,7 +356,7 @@ int32 UGenerateMountPointPayloadManifestCommandlet::Main(const FString& Params)
 				{
 					MountPointIndex++;
 
-					FStringView MountPoint = FStringView(Path).SubStr(MountPointIndex, ContentIndex - MountPointIndex);
+					FStringView MountPoint = FStringView(Path).SubStr(0, ContentIndex);
 
 					if (TArray<FString>* ExistingPaths = MountPointMap.FindByHash(GetTypeHash(MountPoint), MountPoint))
 					{
@@ -371,7 +392,18 @@ int32 UGenerateMountPointPayloadManifestCommandlet::Main(const FString& Params)
 
 	UE_LOG(LogVirtualization, Display, TEXT("Processing mountpoints completed"), MountPointMap.Num());
 
-	if (!OutputMountPointStatistics(OutputFilePath, Stats, bDetailedFilterReasons, ByteFormat))
+	EOutputFlags Flags = EOutputFlags::None;
+	if (bDetailedFilterReasons)
+	{
+		Flags |= EOutputFlags::DetailedFilterReasons;
+	}
+
+	if (bOutputPaths)
+	{
+		Flags |= EOutputFlags::MountPointPaths;
+	}
+
+	if (!OutputMountPointStatistics(OutputFilePath, Stats, Flags, ByteFormat))
 	{
 		return -1;
 	}
@@ -391,6 +423,11 @@ bool UGenerateMountPointPayloadManifestCommandlet::TryParseCmdline(const FString
 		bDetailedFilterReasons = true;
 	}
 
+	if (Switches.Contains(TEXT("NoPaths")))
+	{
+		bOutputPaths = false;
+	}
+
 	FString ByteFormatString;
 	if (FParse::Value(*Params, TEXT("ByteFormat="), ByteFormatString))
 	{
diff --git a/Engine/Source/Editor/VirtualizationEditor/Private/GenerateMountPointPayloadManifestCommandlet.h b/Engine/Source/Editor/VirtualizationEditor/Private/GenerateMountPointPayloadManifestCommandlet.h
--- a/Engine/Source/Editor/VirtualizationEditor/Private/GenerateMountPointPayloadManifestCommandlet.h
+++ b/Engine/Source/Editor/VirtualizationEditor/Private/GenerateMountPointPayloadManifestCommandlet.h
@@ -23,16 +23,19 @@ enum class EByteFormatting
   <project root>/saved/PayloadManifest/mountpoints.csv
  * 
  * Additional args:
- * "-DetailedFilterReasons"
+ * -DetailedFilterReasons
  *     This switch will provide a breakdown of how much content per mount point
  *     is prevented from being virtualized by a specific filter.
- * "-ByteFormat=Bytes/MiB/GiB"
+ * -NoPaths
+ *     This switch will prevent the relative path of the mount point being written
+ *     to the output.
+ * -ByteFormat=Bytes/MiB/GiB
  *     This value allows the caller to set how bytes should be formatted to the
  *     csv file. The default is to output raw bytes.
- * "-OutputPath="
+ * -OutputPath=""
  *     The full path (including filename and extension) of where to write the
  *     final output.
- * "-OutputName="
+ * -OutputName=""
  *     The file name (including extension) of where to write the final output.
  *     Note that the file will be written to the default output directory.
  */
@@ -56,5 +59,7 @@ private:
 	FString OutputFilePath;
 
 	bool bDetailedFilterReasons = false;
+	bool bOutputPaths = true;
+
 	EByteFormatting ByteFormat = EByteFormatting::Bytes;
 };
