Fix minor typo in code documentation.

#rnx

#Synced-CL 49468340

diff --git a/Engine/Source/Editor/VirtualizationEditor/Private/GenerateMountPointPayloadManifestCommandlet.h b/Engine/Source/Editor/VirtualizationEditor/Private/GenerateMountPointPayloadManifestCommandlet.h
--- a/Engine/Source/Editor/VirtualizationEditor/Private/GenerateMountPointPayloadManifestCommandlet.h
+++ b/Engine/Source/Editor/VirtualizationEditor/Private/GenerateMountPointPayloadManifestCommandlet.h
@@ -20,7 +20,7 @@ enum class EByteFormatting
  * -run="VirtualizationEditor.GenerateMountPointPayloadManifestCommandlet"
  * 
  * By default the final output will be written to:
-  <project root>/saved/PayloadManifest/mountpoints.csv
+ * <project root>/saved/PayloadManifest/mountpoints.csv
  * 
  * Additional args:
  * -DetailedFilterReasons
