Improve code documentation for the 'VirtualizationEditor.GeneratePayloadManifest' commandlet.

#rnx

#Synced-CL 49468340

diff --git a/Engine/Source/Editor/VirtualizationEditor/Private/GeneratePayloadManifestCommandlet.h b/Engine/Source/Editor/VirtualizationEditor/Private/GeneratePayloadManifestCommandlet.h
--- a/Engine/Source/Editor/VirtualizationEditor/Private/GeneratePayloadManifestCommandlet.h
+++ b/Engine/Source/Editor/VirtualizationEditor/Private/GeneratePayloadManifestCommandlet.h
@@ -16,6 +16,21 @@
  * Because the commandlet is the VirtualizationEditor module it needs to be invoked 
  * with the command line:
  * -run="VirtualizationEditor.GeneratePayloadManifest"
+ * 
+ * By default the final output will be written to:
+ * <project root>/saved/PayloadManifest/sheet1.csv
+ * note that if the csv file is large enough it will be split over several files.
+ * 
+ * Additional args:
+ * -LocalOnly
+ *     Only output payloads that are stored locally on disk.
+ * -PendingOnly
+ *     Only output payloads that are stored locally on disk but could be virtualized.
+ * -FilteredOnly
+ *     Only output payloads that are stored locally on disk but would never be virtualized
+ *     due to one or more filters.
+ * -VirtualizedOnly
+ *     Only output payloads that have been virtualized.
  */
 UCLASS()
 class UGeneratePayloadManifestCommandlet
@@ -36,9 +51,13 @@ private:
 	enum class EPayloadFilter
 	{
 		None = 0,
+		/* All locally stored payloads */
 		LocalOnly		=  1 << 0,
+		/* Locally stored payloads that could be virtualized */
 		PendingOnly		= (1 << 1) | LocalOnly,
+		/* Local payloads that are filtered from the virtualization process. */
 		FilteredOnly	= (1 << 2) | LocalOnly,
+		/* Virtualized payloads */
 		VirtualizedOnly	= (1 << 3)
 	};
 	FRIEND_ENUM_CLASS_FLAGS(EPayloadFilter);
