Move the console command 'DumpPackagePayloadInfo' to VirtualizationEditor module from Virtualization.

#rnx

- This command only really makes sense in the editor and does not need to exist at runtime.
- This allows us to remove the duplicate LogPackageError code that was in both theConsoleCommandUtils.h and the DumpPackagePayloadInfoCommand file
- Renamed 'LogPackageError' to 'LogPackageTrailerLoadingError' now that it is exposed.
- Renamed 'CCmdDumpPayloadToc' in the places where I had used it (copy paste errors)

#Synced-CL 52606070

diff --git a/Engine/Source/Developer/Virtualization/Private/DumpPackagePayloadInfoCommand.cpp b/Engine/Source/Developer/Virtualization/Private/DumpPackagePayloadInfoCommand.cpp
deleted file mode 100644
diff --git a/Engine/Source/Editor/VirtualizationEditor/Private/ConsoleCommandUtils.cpp b/Engine/Source/Editor/VirtualizationEditor/Private/ConsoleCommandUtils.cpp
--- a/Engine/Source/Editor/VirtualizationEditor/Private/ConsoleCommandUtils.cpp
+++ b/Engine/Source/Editor/VirtualizationEditor/Private/ConsoleCommandUtils.cpp
@@ -11,7 +11,7 @@
 namespace UE::Virtualization
 {
 
-void LogPackageError(FArchive* Ar, const TCHAR* DebugName)
+void LogPackageTrailerLoadingError(FArchive* Ar, const TCHAR* DebugName)
 {
 	if (Ar == nullptr)
 	{
@@ -44,19 +44,19 @@ void LogPackageError(FArchive* Ar, const TCHAR* DebugName)
 	UE_LOGF(LogVirtualization, Error, "Package is tool old (version %d) to have a package trailer (version %d): '%ls'", PackageVersion, int(EUnrealEngineObjectUE5Version::PAYLOAD_TOC), DebugName);
 }
 
-void LogPackageError(const FPackagePath& Path)
+void LogPackageTrailerLoadingError(const FPackagePath& Path)
 {
 	const FString DebugName = Path.GetPackageName();
 	TUniquePtr<FArchive> PackageAr = IPackageResourceManager::Get().OpenReadExternalResource(EPackageExternalResource::WorkspaceDomainFile, DebugName);
 
-	LogPackageError(PackageAr.Get(), *DebugName);
+	LogPackageTrailerLoadingError(PackageAr.Get(), *DebugName);
 }
 
-void LogPackageError(const FString& Path)
+void LogPackageTrailerLoadingError(const FString& Path)
 {
 	TUniquePtr<FArchive> PackageAr(IFileManager::Get().CreateFileReader(*Path));
 
-	LogPackageError(PackageAr.Get(), *Path);
+	LogPackageTrailerLoadingError(PackageAr.Get(), *Path);
 }
 
 TArray<TPair<FString, UE::FPackageTrailer>> LoadPackageTrailerFromArgs(const TArray<FString>& Args)
@@ -84,7 +84,7 @@ TArray<TPair<FString, UE::FPackageTrailer>> LoadPackageTrailerFromArgs(const TAr
 		{
 			if (!FPackageTrailer::TryLoadFromPackage(Path, Trailer))
 			{
-				LogPackageError(Path);
+				LogPackageTrailerLoadingError(Path);
 				continue;
 			}
 		}
@@ -95,7 +95,7 @@ TArray<TPair<FString, UE::FPackageTrailer>> LoadPackageTrailerFromArgs(const TAr
 
 			if (!FPackageTrailer::TryLoadFromFile(PathString, Trailer))
 			{
-				LogPackageError(PathString);
+				LogPackageTrailerLoadingError(PathString);
 				continue;
 			}
 		}
diff --git a/Engine/Source/Editor/VirtualizationEditor/Private/ConsoleCommandUtils.h b/Engine/Source/Editor/VirtualizationEditor/Private/ConsoleCommandUtils.h
--- a/Engine/Source/Editor/VirtualizationEditor/Private/ConsoleCommandUtils.h
+++ b/Engine/Source/Editor/VirtualizationEditor/Private/ConsoleCommandUtils.h
@@ -3,11 +3,17 @@
 #pragma once
 
 #include "Containers/Array.h"
+#include "Containers/UnrealString.h"
 #include "UObject/PackageTrailer.h"
 
+class FPackagePath;
+
 namespace UE::Virtualization
 {
 
 TArray<TPair<FString, UE::FPackageTrailer>> LoadPackageTrailerFromArgs(const TArray<FString>& Args);
 
+void LogPackageTrailerLoadingError(const FPackagePath& Path);
+void LogPackageTrailerLoadingError(const FString& Path);
+
 } //namespace UE::Virtualization
diff --git a/Engine/Source/Editor/VirtualizationEditor/Private/DumpPackagePayloadInfoCommand.cpp b/Engine/Source/Editor/VirtualizationEditor/Private/DumpPackagePayloadInfoCommand.cpp
new file mode 100644
--- /dev/null
+++ b/Engine/Source/Editor/VirtualizationEditor/Private/DumpPackagePayloadInfoCommand.cpp
@@ -0,0 +1,149 @@
+// Copyright Epic Games, Inc. All Rights Reserved.
+
+#include "ConsoleCommandUtils.h"
+#include "HAL/FileManager.h"
+#include "HAL/IConsoleManager.h"
+#include "Logging/LogMacros.h"
+#include "Misc/PackageName.h"
+#include "Misc/PackagePath.h"
+#include "Misc/ScopedSlowTask.h"
+#include "UObject/PackageFileSummary.h"
+#include "UObject/PackageResourceManager.h"
+#include "UObject/PackageTrailer.h"
+#include "VirtualizationExperimentalUtilities.h"
+
+#define LOCTEXT_NAMESPACE "Virtualization"
+
+namespace UE::Virtualization
+{
+
+#if WITH_EDITORONLY_DATA
+
+FString BytesToString(int64 SizeInBytes)
+{
+	if (SizeInBytes < (8 *1024))
+	{
+		return FString::Printf(TEXT("%4" INT64_FMT " bytes"), SizeInBytes);
+	}
+	else if (SizeInBytes < (1024 * 1024))
+	{
+		double SizeInKb = static_cast<double>(SizeInBytes) / (1024.0);
+		return FString::Printf(TEXT("%.2f KB"), SizeInKb);
+	}
+	else
+	{
+		double SizeInMB = static_cast<double>(SizeInBytes) / (1024.0 * 1024.0);
+		return FString::Printf(TEXT("%.2f MB"), SizeInMB);
+	}
+}
+
+/**
+ * This function is used to write information about package's payloads to the log file. This has no 
+ * practical development use and should only be used for debugging purposes. 
+ * 
+ * @param Args	The function expects each arg to be a valid package path. Failure to provide a valid
+ *				package path will result in errors being written to the log.
+ */
+void DumpPackagePayloadInfo(const TArray<FString>& Args)
+{
+	if (Args.Num() == 0)
+	{
+		UE_LOGF(LogVirtualization, Error, "Command 'DumpPackagePayloadInfo' called without any arguments");
+		return;
+	}
+
+	TRACE_CPUPROFILER_EVENT_SCOPE(DumpPackagePayloadInfo);
+	FScopedSlowTask Progress(0.0f, LOCTEXT("VAFindPayloadInfo", "Finding payload info..."));
+
+	for (const FString& Arg : Args)
+	{
+		FString PathString;
+
+		if (FPackageName::ParseExportTextPath(Arg, nullptr /*OutClassName*/, &PathString))
+		{
+			PathString = FPackageName::ObjectPathToPackageName(PathString);
+		}
+		else
+		{
+			PathString = Arg;
+		}
+
+		FPackageTrailer Trailer;
+
+		FPackagePath Path;
+		if (FPackagePath::TryFromMountedName(PathString, Path))
+		{	
+			if (!FPackageTrailer::TryLoadFromPackage(Path, Trailer))
+			{
+				LogPackageTrailerLoadingError(Path);
+				continue;
+			}
+		}
+		else if (IFileManager::Get().FileExists(*PathString))
+		{
+			// IF we couldn't turn it into a FPackagePath it could be a path to a package not under any current mount point.
+			// So for a final attempt we will see if we can find the file on disk and load the package trailer that way.
+			
+			if (!FPackageTrailer::TryLoadFromFile(PathString, Trailer))
+			{
+				LogPackageTrailerLoadingError(PathString);
+				continue;
+			}	
+		}
+		else
+		{
+			UE_LOGF(LogVirtualization, Error, "Arg '%ls' could not be converted to a valid package path", *Arg);
+			continue;
+		}
+
+		TArray<FIoHash> LocalPayloadIds = Trailer.GetPayloads(UE::EPayloadStorageType::Local);
+		TArray<FIoHash> VirtualizedPayloadIds = Trailer.GetPayloads(UE::EPayloadStorageType::Virtualized);
+
+		UE_LOGF(LogVirtualization, Display, ""); // Blank line to make the output easier to read
+		UE_LOGF(LogVirtualization, Display, "Package: '%ls' has %d local and %d virtualized payloads", *Path.GetDebugName(), LocalPayloadIds.Num(), VirtualizedPayloadIds.Num());
+
+		if (LocalPayloadIds.Num() > 0)
+		{
+			UE_LOGF(LogVirtualization, Display, "LocalPayloads:");
+			UE_LOGF(LogVirtualization, Display, "Index | %-40ls | SizeOnDisk | FilterReason", TEXT("PayloadIdentifier"));
+			for (int32 Index = 0; Index < LocalPayloadIds.Num(); ++Index)
+			{
+				FPayloadInfo Info = Trailer.GetPayloadInfo(LocalPayloadIds[Index]);
+
+				Info.FilterFlags = UE::Virtualization::Utils::FixFilterFlags(PathString, Info.CompressedSize, Info.FilterFlags);
+
+				UE_LOGF(LogVirtualization, Display, "%02d    | %ls | %-10ls | %ls",
+					Index,
+					*LexToString(LocalPayloadIds[Index]),
+					*BytesToString(Info.CompressedSize),
+					*LexToString(Info.FilterFlags));
+			}
+		}
+
+		if (VirtualizedPayloadIds.Num() > 0)
+		{
+			UE_LOGF(LogVirtualization, Display, "VirtualizedPayloads:");
+			UE_LOGF(LogVirtualization, Display, "Index|\t%-40ls|\tFilterReason", TEXT("PayloadIdentifier"));
+			for (int32 Index = 0; Index < VirtualizedPayloadIds.Num(); ++Index)
+			{
+				FPayloadInfo Info = Trailer.GetPayloadInfo(VirtualizedPayloadIds[Index]);
+				UE_LOGF(LogVirtualization, Display, "%02d:  |\t%ls|\t%ls", Index, *LexToString(VirtualizedPayloadIds[Index]), *LexToString(Info.FilterFlags));
+			}
+		}
+	}
+}
+
+/** 
+ * Note that this command is only valid when 'WITH_EDITORONLY_DATA 1' as virtualized payloads are not
+ * expected to exist at runtime. 
+ */
+static FAutoConsoleCommand CCmdDumpPayloadTrailer = FAutoConsoleCommand(
+	TEXT("DumpPackagePayloadInfo"),
+	TEXT("Writes out information about a package's payloads to the log."),
+	FConsoleCommandWithArgsDelegate::CreateStatic(DumpPackagePayloadInfo));
+
+#endif //WITH_EDITORONLY_DATA
+
+} // namespace UE::Virtualization
+
+#undef LOCTEXT_NAMESPACE
diff --git a/Engine/Source/Editor/VirtualizationEditor/Private/ValidatePackagePayloadsConsoleCommand.cpp b/Engine/Source/Editor/VirtualizationEditor/Private/ValidatePackagePayloadsConsoleCommand.cpp
--- a/Engine/Source/Editor/VirtualizationEditor/Private/ValidatePackagePayloadsConsoleCommand.cpp
+++ b/Engine/Source/Editor/VirtualizationEditor/Private/ValidatePackagePayloadsConsoleCommand.cpp
@@ -114,7 +114,7 @@ void ValidatePackagePayloads(const TArray<FString>& Args)
 	}
 }
 
-static FAutoConsoleCommand CCmdDumpPayloadToc = FAutoConsoleCommand(
+static FAutoConsoleCommand CCmdValidatePackagePayloads = FAutoConsoleCommand(
 	TEXT("ValidatePackagePayloads"),
 	TEXT("Checks all payloads that a package references and makes sure that they are valid"),
 	FConsoleCommandWithArgsDelegate::CreateStatic(ValidatePackagePayloads));
